Asset Size Mapper tool
- New "Size Map" feature that shows the resource size for selected assets in Content Browser
- Select assets in Content Browser, right click and choose "Size Map..." to show sizes for those assets
- Select actors and press Alt+Shift+M to show sizes of assets used by those actors
- Select assets in Reference Viewer, right click and choose "Size Map..." to show sizes for those assets
- Also added a new "Tree Map" general purpose Slate widget (STreeMap)

diff --git a/Engine/Source/Developer/TreeMap/Public/ITreeMap.h b/Engine/Source/Developer/TreeMap/Public/ITreeMap.h
new file mode 100644
--- /dev/null
+++ b/Engine/Source/Developer/TreeMap/Public/ITreeMap.h
@@ -0,0 +1,293 @@
+// Copyright 1998-2013 Epic Games, Inc. All Rights Reserved.
+
+#pragma once
+
+
+typedef TSharedRef< class FTreeMapAttributeData > FTreeMapAttributeDataRef;
+typedef TSharedPtr< class FTreeMapAttributeData > FTreeMapAttributeDataPtr;
+
+
+/**
+ * Stores data for a specific attribute type on a node
+ */
+class FTreeMapAttributeData
+{
+
+public:
+
+	/** Value for this data */
+	FName Value;
+
+
+public:
+
+	/** Default constructor for FTreeMapAttributeData */
+	FTreeMapAttributeData()
+		: Value( NAME_None )
+	{
+	}
+
+	/** Intitialization constructor for FTreeMapAttributeData */
+	FTreeMapAttributeData( const FName InitValue )
+		: Value( InitValue )
+	{
+	}
+};
+
+
+
+typedef TSharedRef< class FTreeMapNodeData > FTreeMapNodeDataRef;
+typedef TSharedPtr< class FTreeMapNodeData > FTreeMapNodeDataPtr;
+
+/**
+ * Single node in a tree map, which may have any number of child nodes, each with their own children and so on
+ */
+class FTreeMapNodeData : public TSharedFromThis<FTreeMapNodeData>
+{
+
+public:
+
+	/** Node name.  When drawing the tree map, this shows up at the top of the inside of the node's rectangle */
+	FString Name;
+
+	/** Node name (line 2).  Left nodes only.  This shows up on the second line of the node, right underneath its name.  Might be drawn in a smaller font by default. */
+	FString Name2;
+
+	/** Center info text.  Leaf nodes only.  This draws right in the center of the node, with a larger font by default */
+	FString CenterText;
+
+
+	/** 
+	 * Size of this node.  This works differently depending on whether the node is a leaf node (no children) or a container node (has children):
+	 *  - Leaf nodes must ALWAYS have a non-zero size.  
+	 *	- Container nodes with size of zero will have their size determined automatically by the sum of their child sizes
+	 *  - Container nodes with a non-zero size will override their child sizes and always have the explicitly set size.  Sizes of children will still be used to proportionately scale the size of the child nodes within this container.
+	 */
+	float Size;
+
+	/** Background brush for this node's box (optional) */
+	const struct FSlateBrush* BackgroundBrush;
+
+	/** Color for this node.  This will be set automatically unless you use a customization. */
+	FLinearColor Color;
+
+	/** Hashtags for this node.  This is just meta-data for the node that may have been loaded from a file.  It is not used by the TreeMap system 
+	   directly, but you could harvest these tags and convert them into custom attributes if you wanted to! */
+	TArray<FString> HashTags;
+
+	/** Map of attribute name to it's respective bit of data */
+	TMap< FName, FTreeMapAttributeDataPtr > Attributes;
+
+	/** Back pointer to parent node, or NULL if no parent exists (root node) */
+	FTreeMapNodeData* Parent;
+
+	/** List of child nodes */
+	TArray< FTreeMapNodeDataPtr > Children;
+
+
+public:
+
+	/** Default constructor for FTreeMapNodeData */
+	FTreeMapNodeData()
+		: Name(),
+		  Name2(),
+		  CenterText(),
+		  Size( 0.0 ),
+		  BackgroundBrush( nullptr ),
+		  Color( FLinearColor::White ),
+		  HashTags(),
+		  Parent( NULL ),
+		  Children()
+	{
+	}
+
+	/** @return Returns true if this is a leaf node */
+	bool IsLeafNode() const
+	{
+		return Children.Num() == 0;
+	}
+
+	/** @return Copies this node into another node (not children, though.  The copied node will have no children.) */
+	void CopyNodeInto( FTreeMapNodeData& NodeCopy ) const
+	{
+		NodeCopy.Name = this->Name;
+		NodeCopy.Name2 = this->Name2;
+		NodeCopy.CenterText = this->CenterText;
+		NodeCopy.BackgroundBrush = this->BackgroundBrush;
+		NodeCopy.Size = this->Size;
+		NodeCopy.Color = this->Color;
+		NodeCopy.HashTags = this->HashTags;
+
+		for( auto HashIter( this->Attributes.CreateConstIterator() ); HashIter; ++HashIter )
+		{
+			FTreeMapAttributeDataRef AttributeDataCopy( new FTreeMapAttributeData() );
+			AttributeDataCopy->Value = HashIter.Value()->Value;
+			NodeCopy.Attributes.Add( HashIter.Key(), AttributeDataCopy );
+		}
+
+		NodeCopy.Children.Empty();
+		NodeCopy.Parent = this->Parent;
+	}
+
+	/** @return Returns a copy of this node; all child nodes are copied, too! */
+	FTreeMapNodeDataRef CopyNodeRecursively() const
+	{
+		FTreeMapNodeDataRef NodeCopy = MakeShareable( new FTreeMapNodeData() );
+
+		this->CopyNodeInto( *NodeCopy );
+
+		for( const auto& ChildNode : this->Children )
+		{
+			const auto& ChildCopy = ChildNode->CopyNodeRecursively();
+			NodeCopy->Children.Add( ChildCopy );
+			ChildCopy->Parent = &NodeCopy.Get();
+		}
+
+		return NodeCopy;
+	}
+};
+
+
+/** Type of tree map */
+enum class ETreeMapType
+{
+	/** Plain tree map */
+	Standard,
+
+	/** Squarified tree map */
+	Squarified,
+};
+
+
+/**
+ * Configuration for a new tree map
+ */
+struct FTreeMapOptions
+{
+	/** Width of whole display area */
+	float DisplayWidth;
+
+	/** Height of whole display area */
+	float DisplayHeight;
+
+	/** Type of tree map */
+	ETreeMapType TreeMapType;
+
+	/** Font to use for titles.  This is needed so that we can calculate the amount of padding space needed.  This font will be used for text
+	    titles at the root-most level of the tree.  Nested levels get smaller fonts.  This also affects the amount of space reserved for the title area at the top of each node. */
+	FSlateFontInfo NameFont;
+
+	/** Font for second line of text, under the title.  Leaf nodes only.  Usually a bit smaller.  Works just like NameFont */
+	FSlateFontInfo Name2Font;
+
+	/** Font for any text that's centered inside the middle of the node.  Leaf nodes only.  Usually a bit larger.  Works just like NameFont */
+	FSlateFontInfo CenterTextFont;
+
+	/** Number of font sizes to drop with each depth level of the tree */
+	int32 FontSizeChangeBasedOnDepth;
+
+	/** Padding around the outside of top-level container nodes in the tree */
+	float TopLevelContainerOuterPadding;
+
+	/** Padding around the outside of nested container nodes in the tree */
+	float NestedContainerOuterPadding;
+
+	/** Padding around a set of children inside of containers, under the container's title area */
+	float ContainerInnerPadding;
+
+	/** Minimize size of a tree node that will be allowed to have a title and padding.  Nodes smaller than this will be a simple rectangle.  Doesn't apply to top level nodes. */
+	float MinimumInteractiveNodeSize;
+
+
+	/** Default constructor for FTreeMapOptions that initializes good defaults */
+	FTreeMapOptions()
+		: DisplayWidth( 1.0f ),
+		  DisplayHeight( 1.0f ),
+		  TreeMapType( ETreeMapType::Standard ),
+		  NameFont( FCoreStyle::Get().GetFontStyle( TEXT( "NormalText" ) ) ),
+		  Name2Font( FCoreStyle::Get().GetFontStyle( TEXT( "NormalText" ) ) ),
+		  CenterTextFont( FCoreStyle::Get().GetFontStyle( TEXT( "NormalText" ) ) ),
+		  FontSizeChangeBasedOnDepth( 1 ),
+		  TopLevelContainerOuterPadding( 0.0f ),
+		  NestedContainerOuterPadding( 0.0f ),
+		  ContainerInnerPadding( 0.0f ),
+		  MinimumInteractiveNodeSize( 0.0f )
+	{
+		NameFont.Size = 12;
+		Name2Font.Size = 8;
+		CenterTextFont.Size = 24;
+	}
+};
+
+
+
+/**
+ * Visual ID for a node, generated by the tree map system
+ */
+struct FTreeMapNodeVisualInfo 
+{
+	/** Pointer back to the node data this visual was originally created from */
+	class FTreeMapNodeData* NodeData;
+
+	/** Position for this node */
+	FVector2D Position;
+
+	/** Size for this node */
+	FVector2D Size;
+
+	/** Node color */
+	FLinearColor Color;
+
+	/** Node fonts */
+	FSlateFontInfo NameFont;
+	FSlateFontInfo Name2Font;
+	FSlateFontInfo CenterTextFont;
+
+	/** True if the node is 'interactive'.  That is, we have space for the node's title and for it to be clicked on. */
+	bool bIsInteractive;
+};
+
+
+/**
+ * Public tree map interface
+ */
+class ITreeMap
+{
+
+public:
+
+	/**
+	 * Gets the visuals for this tree map.  Be careful not to destroy the original tree while still using the visual objects, otherwise
+	 * the back pointers to the original tree nodes will be trash!  Visuals are ordered such that nested visuals appear later in the list
+	 * than their parents (so they can be drawn back to front, etc.)
+	 *
+	 * @return List of visuals
+	 */
+	virtual TArray<FTreeMapNodeVisualInfo> GetVisuals() = 0;
+
+	/**
+	 * Static: Creates a tree map object, given tree node source data
+	 *
+	 * @param	Options			Configuration options for this tree map
+	 * @param	RootNodeData	The source data for the root of the tree
+	 *
+	 * @return	The newly-created ITreeMap instance 
+	 */
+	static TSharedRef<ITreeMap> CreateTreeMap( const FTreeMapOptions& Options, const FTreeMapNodeDataRef& RootNodeData );
+
+	/**
+	 * Parses an OPML XML document and converts the outliner content to tree map node data
+	 *
+	 * @param	OPMLFilePath		Path to the OPML file on disk
+	 * @param	OutErrorMessage		Error message string in the case that something went wrong
+	 *
+	 * @return	The root tree node for the parsed data, or an invalid pointer if something went wrong
+	 */
+	static FTreeMapNodeDataPtr ParseOPMLToTreeMapData( const FString& OPMLFilePath, FString& OutErrorMessage );
+
+	/** Virtual destructor */
+	virtual ~ITreeMap() {}
+};
+
+
+
diff --git a/Engine/Source/Developer/TreeMap/Public/ITreeMapCustomization.h b/Engine/Source/Developer/TreeMap/Public/ITreeMapCustomization.h
new file mode 100644
--- /dev/null
+++ b/Engine/Source/Developer/TreeMap/Public/ITreeMapCustomization.h
@@ -0,0 +1,122 @@
+// Copyright 1998-2013 Epic Games, Inc. All Rights Reserved.
+
+#pragma once
+
+
+#include "ITreeMap.h"
+
+
+typedef TSharedPtr< class FTreeMapAttributeValue > FTreeMapAttributeValuePtr;
+typedef TSharedRef< class FTreeMapAttributeValue > FTreeMapAttributeValueRef;
+
+/**
+ * Describes a value for an attribute of a tree node
+ */
+class FTreeMapAttributeValue
+{
+
+public:
+
+	/** Name of this value */
+	FName Name;
+
+	/** The node size that this value maps to */
+	float NodeSize;
+
+	/** The node color that this value maps to */
+	FLinearColor NodeColor;
+
+
+public:
+
+	/** Default constructor */
+	FTreeMapAttributeValue()
+		: Name(),
+		  NodeSize( 1.0f ),
+		  NodeColor( FLinearColor::White )
+	{
+	}
+
+
+	/** Virtual destructor */
+	virtual ~FTreeMapAttributeValue()
+	{
+	}
+};
+
+
+
+typedef TSharedPtr< class FTreeMapAttribute > FTreeMapAttributePtr;
+typedef TSharedRef< class FTreeMapAttribute > FTreeMapAttributeRef;
+
+/**
+ * Describes a customized attribute of a tree node
+ */
+class FTreeMapAttribute
+{
+
+public:
+
+	/** Name of this attribute */
+	FName Name;
+
+	/** Maps a value name to the data that describes that value */
+	TMap< FName, FTreeMapAttributeValuePtr > Values;
+
+	/** Default value to use when none is specified on a node */
+	FTreeMapAttributeValuePtr DefaultValue;
+
+
+public:
+
+	/** Default constructor */
+	FTreeMapAttribute()
+		: Name(),
+		  Values(),
+		  DefaultValue()
+	{
+	}
+
+	/** Virtual destructor */
+	virtual ~FTreeMapAttribute()
+	{
+	}
+
+};
+
+
+
+/**
+ * Implement ITreeMapCustomization and pass it along when creating your tree map to allow for custom attributes and formatting features!
+ */
+class ITreeMapCustomization
+{
+
+public:
+
+	/** @return Returns the name of this customization */
+	virtual FName GetName() const = 0;
+
+	/** @return Returns all of the possible attributes in this customization.  Each attribute defines a set of possible values for that attribute. */
+	virtual const TMap< FName, FTreeMapAttributePtr >& GetAttributes() const = 0;
+
+	/** @return Returns the default attribute type to size the nodes based on */
+	virtual FTreeMapAttributePtr GetDefaultSizeByAttribute() const { return NULL; }
+
+	/** @return Returns the default attribute type to color the nodes based on */
+	virtual FTreeMapAttributePtr GetDefaultColorByAttribute() const { return NULL; }
+
+	/**
+	 * Optional override that can be called to convert hash tags on tree nodes into proper attribute values.  This is useful if you've loaded
+	 * tree node data from a file and want to take the raw hash tag strings and make values for those.
+	 *
+	 * @param	Node	The root node to process
+	 */
+	virtual void ProcessHashTagsRecursively( const FTreeMapNodeDataRef& Node ) const {};
+
+	/** Virtual destructor */
+	virtual ~ITreeMapCustomization() {}
+};
+
+
+
diff --git a/Engine/Source/Developer/TreeMap/Public/STreeMap.h b/Engine/Source/Developer/TreeMap/Public/STreeMap.h
new file mode 100644
--- /dev/null
+++ b/Engine/Source/Developer/TreeMap/Public/STreeMap.h
@@ -0,0 +1,324 @@
+// Copyright 1998-2013 Epic Games, Inc. All Rights Reserved.
+
+#pragma once
+
+#include "ITreeMap.h"
+#include "ITreeMapCustomization.h"
+#include "TreeMapStyle.h"
+
+
+/**
+ * Graphical tree map widget with interactive controls
+ */
+class TREEMAP_API STreeMap : public SLeafWidget
+{
+
+public:
+
+	/** Optional delegate that fires when the node is double-clicked in the tree.  If you don't override this, the tree will use its
+	    default handling of double-click, which will re-root the tree on the node under the cursor */
+	DECLARE_DELEGATE_OneParam( FOnTreeMapNodeDoubleClicked, FTreeMapNodeData& );
+
+
+
+	SLATE_BEGIN_ARGS( STreeMap )
+		: _AllowEditing( false )
+		, _BackgroundImage( FTreeMapStyle::Get().GetBrush( "TreeMap.Background" ) )
+		, _NodeBackground( FTreeMapStyle::Get().GetBrush( "TreeMap.NodeBackground" ) )
+		, _HoveredNodeBackground( FTreeMapStyle::Get().GetBrush( "TreeMap.HoveredNodeBackground" ) )
+		, _NameFont()
+		, _Name2Font()
+		, _CenterTextFont()
+		, _BorderPadding( FTreeMapStyle::Get().GetVector( "TreeMap.BorderPadding" ) )
+		, _MinimumVisualTreeNodeSize( 64 * 64 )
+		, _NavigationTransitionTime( 0.25f )
+		, _TopLevelContainerOuterPadding( 4.0f )
+		, _NestedContainerOuterPadding( 0.0f )
+		, _ContainerInnerPadding( 4.0f )
+		, _ChildContainerTextPadding( 2.0f )
+	{}
+
+		/** Sets whether the user can edit the tree map interactively by dragging nodes around and typing new node labels */
+		SLATE_ATTRIBUTE( bool, AllowEditing )
+
+		/** Background image to use for the tree map canvas area */
+		SLATE_ATTRIBUTE( const FSlateBrush*, BackgroundImage )
+
+		/** Background to use for each tree node */
+		SLATE_ATTRIBUTE( const FSlateBrush*, NodeBackground )
+			
+		/** Background to use for nodes that the mouse is hovering over */
+		SLATE_ATTRIBUTE( const FSlateBrush*, HoveredNodeBackground )
+
+		/** Sets the font used to draw the text.  Note the height of the font may change automatically based on the tree node size */
+		SLATE_ATTRIBUTE( FSlateFontInfo, NameFont )
+			
+		/** Font for second line of text, under the title.  Leaf nodes only.  Usually a bit smaller.  Works just like NameFont */
+		SLATE_ATTRIBUTE( FSlateFontInfo, Name2Font );
+
+		/** Font for any text that's centered inside the middle of the node.  Leaf nodes only.  Usually a bit larger.  Works just like NameFont */
+		SLATE_ATTRIBUTE( FSlateFontInfo, CenterTextFont );
+
+		/** Border Padding around fill bar */
+		SLATE_ATTRIBUTE( FVector2D, BorderPadding )
+
+		/** Minimum size in pixels of a tree node that we should bother including in the UI.  Below this size, you'll need to drill down to see the node. */
+		SLATE_ARGUMENT( int32, MinimumVisualTreeNodeSize );
+
+		/** How many seconds to animate the visual transition when the user navigates to a new tree node, or after a modification of the tree takes place */
+		SLATE_ARGUMENT( float, NavigationTransitionTime );
+
+		/** How many pixels of padding around the outside of the root-level tree node box.  Adding padding makes the tree look better, but you lose some sizing accuracy */
+		SLATE_ARGUMENT( float, TopLevelContainerOuterPadding );
+
+		/** How many pixels of padding around the outside of a non-root tree node's box.  Adding padding makes the tree look better, but you lose some sizing accuracy */
+		SLATE_ARGUMENT( float, NestedContainerOuterPadding );
+
+		/** How many pixels of spacing between the container and its child containers.  Adding padding makes the tree look better, but you lose some sizing accuracy */
+		SLATE_ARGUMENT( float, ContainerInnerPadding );
+
+		/** How many pixels to pad text that's drawn inside of a container (not the top-level container, though) */
+		SLATE_ARGUMENT( float, ChildContainerTextPadding );
+
+		/** Optional delegate that fires when the node is double-clicked in the tree.  If you don't override this, the tree will use its
+			default handling of double-click, which will re-root the tree on the node under the cursor */
+		SLATE_EVENT( FOnTreeMapNodeDoubleClicked, OnTreeMapNodeDoubleClicked )
+
+	SLATE_END_ARGS()
+
+	/**
+	 * Construct the widget
+	 *
+	 * @param	InArgs				A declaration from which to construct the widget
+	 * @param	InTreeMapNodeData	The node data we'll be visualizing and editing (root node)
+	 * @param	InCustomization		An optional 'customization' of the tree that allows for filtering and displaying the tree map data according to various attributes
+	 */
+	void Construct( const FArguments& InArgs, const TSharedRef<class FTreeMapNodeData>& InTreeMapNodeData, const TSharedPtr< ITreeMapCustomization >& InCustomization );
+
+	/** SWidget overrides */
+	virtual int32 OnPaint( const FPaintArgs& Args, const FGeometry& AllottedGeometry, const FSlateRect& MyClippingRect, FSlateWindowElementList& OutDrawElements, int32 LayerId, const FWidgetStyle& InWidgetStyle, bool bParentEnabled ) const override;
+	virtual FVector2D ComputeDesiredSize(float LayoutScaleMultiplier) const override;
+	virtual void Tick( const FGeometry& AllottedGeometry, const double InCurrentTime, const float InDeltaTime ) override;
+	virtual FReply OnMouseButtonDown( const FGeometry& InMyGeometry, const FPointerEvent& InMouseEvent ) override;
+	virtual FReply OnMouseButtonUp( const FGeometry& InMyGeometry, const FPointerEvent& InMouseEvent ) override;
+	virtual FReply OnMouseMove( const FGeometry& InMyGeometry, const FPointerEvent& InMouseEvent ) override;
+	virtual void OnMouseLeave( const FPointerEvent& InMouseEvent ) override;
+	virtual FReply OnMouseButtonDoubleClick( const FGeometry& InMyGeometry, const FPointerEvent& InMouseEvent ) override;
+	virtual FReply OnMouseWheel( const FGeometry& MyGeometry, const FPointerEvent& MouseEvent ) override;
+	virtual bool SupportsKeyboardFocus() const override;
+	virtual FReply OnKeyDown( const FGeometry& MyGeometry, const FKeyEvent& InKeyboardEvent ) override;
+
+	/**
+	 * Sets a new "active root" for the tree.  This is used to "drill down" to child tree nodes or "climb up" back to the root.
+	 * This also happens automatically when the user presses the mouse wheel to zoom in and out over nodes.
+	 * 
+	 * @param	NewRoot						The tree's new root.  This can be any node within the tree
+	 * @param	bShouldPlayTransition		If enabled, an animation will play
+	 */
+	void SetTreeRoot( const FTreeMapNodeDataRef& NewRoot, const bool bShouldPlayTransition );
+
+	/**
+	 * Refreshes the tree map from its source data.  The current active root will remain as the base of the tree.  Call this
+	 * after you've changed data in the tree.
+	 * 
+	 * @param	bShouldPlayTransition		If enabled, an animation will play
+	 */
+	void RebuildTreeMap( const bool bShouldPlayTransition );
+		
+
+protected:
+
+	/** Finds the node visual that's under the cursor */
+	struct FTreeMapNodeVisualInfo* FindNodeVisualUnderCursor( const FGeometry& MyGeometry, const FVector2D& ScreenSpaceCursorPosition );
+
+	/** Blends between the current visual state and the state before the last navigation transition, by the specified amount */
+	void MakeBlendedNodeVisual( const int32 VisualIndex, const float NavigationAlpha, FTreeMapNodeVisualInfo& OutVisual ) const;
+
+	/** @return True if we're playing a transition after a navigation action (zoom!) */
+	bool IsNavigationTransitionActive() const;
+
+	/** Grabs a snapshot of the current tree so we can restore it if the user performs an Undo action */
+	void TakeUndoSnapshot();
+
+	/** Undoes the last action */
+	void Undo();
+
+	/** Redoes the last action */
+	void Redo();
+
+	/** Reparents DroppedNode to NewParentNode (undoable!) */
+	void ReparentNode( const FTreeMapNodeDataRef DroppedNode, const FTreeMapNodeDataRef NewParentNode );
+
+	/** Deletes the node under the mouse cursor, if any (undoable!) */
+	FReply DeleteHoveredNode();
+
+	/** Insert a new node as a child of the node under the cursor */
+	FReply InsertNewNodeAsChildOfHoveredNode( const FGeometry& MyGeometry );
+
+	/** Searches for the specified node in an identical copy of the node tree. */
+	FTreeMapNodeDataPtr FindNodeInCopiedTree( const FTreeMapNodeDataRef& NodeToFind, const FTreeMapNodeDataRef& OriginalNode, const FTreeMapNodeDataRef& CopiedRootNode ) const;
+
+	/** Pops up a box to allow the user to start renaming a node's title (undoable!) */
+	void StartRenamingNode( const FGeometry& MyGeometry, const FTreeMapNodeDataRef& NodeData, const FVector2D& RelativePosition, const bool bIsNewNode );
+
+	/** Called when the user commits a rename change */
+	void RenamingNode_OnTextCommitted( const FText& NewText, ETextCommit::Type, TSharedRef<FTreeMapNodeData> NodeToRename );
+
+	/** Stops renaming a node, committing whatever text was entered */
+	void StopRenamingNode();
+
+	/** Called to apply the current 'size based on' and 'color based on' settings to the tree */
+	void ApplyVisualizationToNodes( const FTreeMapNodeDataRef& Node );
+
+	/** Recursively applies the active visualization to all nodes, such as size by attribute, or color by attribute */
+	void ApplyVisualizationToNodesRecursively( const FTreeMapNodeDataRef& Node, const FLinearColor& DefaultColor, const int32 TreeDepth );
+		
+	/** Displays a context menu at the specified location with options for configuring the tree display */
+	void ShowOptionsMenuAt( const FVector2D& ScreenSpacePosition );
+
+
+private:
+
+	/**
+	 * TreeMap data
+	 */
+
+	/** The tree map that we're presenting using this widget */
+	TSharedPtr<class ITreeMap> TreeMap;
+
+	/** Customization for this tree map */
+	TSharedPtr<class ITreeMapCustomization> Customization;
+
+	/** The previous tree map for the last view we were looking at */
+	TSharedPtr<class ITreeMap> LastTreeMap;
+
+	/** The source data for the tree map */
+	FTreeMapNodeDataPtr TreeMapNodeData;
+
+	/** Current root of the tree that we're visualizing.  The user can re-root the tree by double-clicking on tree map nodes */
+	FTreeMapNodeDataPtr ActiveRootTreeMapNode;
+
+	/** Which attribute to visualize the size of the nodes based on, or null for default sizes */
+	FTreeMapAttributePtr SizeNodesByAttribute;
+
+	/** Which attribute to visualize the color of the nodes based on, or null for automatic colors */
+	FTreeMapAttributePtr ColorNodesByAttribute;
+
+
+	/**
+	 * Widget layout
+	 */
+
+	/** Background image to use for the tree map canvas area */
+	TAttribute< const FSlateBrush* > BackgroundImage;
+
+	/** Background to use for each tree node */
+	TAttribute< const FSlateBrush* > NodeBackground;
+
+	/** Background to use for nodes that the mouse is hovering over */
+	TAttribute< const FSlateBrush* > HoveredNodeBackground;
+
+	/** Border Padding */
+	TAttribute<FVector2D> BorderPadding;
+
+	/** The font name used to draw the text */
+	TAttribute< FSlateFontInfo > NameFont;
+
+	/** Font for second line of text, under the title.  Leaf nodes only.  Usually a bit smaller.  Works just like NameFont */
+	TAttribute< FSlateFontInfo > Name2Font;
+
+	/** Font for any text that's centered inside the middle of the node.  Leaf nodes only.  Usually a bit larger.  Works just like NameFont */
+	TAttribute< FSlateFontInfo > CenterTextFont;
+
+
+	/**
+	 * Navigation
+	 */
+
+	/** The visual that the mouse was over last */
+	FTreeMapNodeVisualInfo* MouseOverVisual;
+
+	/** Optional delegate that fires when the node is double-clicked in the tree.  If you don't override this, the tree will use its
+		default handling of double-click, which will re-root the tree on the node under the cursor */
+	FOnTreeMapNodeDoubleClicked OnTreeMapNodeDoubleClicked;
+
+
+	/**
+	 * TreeMap visuals
+	 */
+
+	/** Minimum size in pixels of a tree node that we should bother including in the UI.  Below this size, you'll need to drill down to see the node. */
+	int32 MinimumVisualTreeNodeSize;
+
+	/** Size of our geometry the last time we rebuilt the tree map */
+	FVector2D TreeMapSize;
+
+	/** Cached set of visuals for the tree map the last time it was generated */
+	TArray<struct FTreeMapNodeVisualInfo> CachedNodeVisuals;
+
+	/** Cached node visuals for last tree map */
+	TArray<struct FTreeMapNodeVisualInfo> LastCachedNodeVisuals;
+
+	/** Maps node visual indices to their last node visual indices.  This is used for navigation transitions */
+	TMap<int32, int32> NodeVisualIndicesToLastIndices;
+
+	/** List of node visual indices which no longer map to a node in the current layout. This is used for navigation transitions */
+	TArray<int32> OrphanedLastIndices;
+
+	/** Time that tree view was changed last after navigation */
+	FCurveSequence NavigateAnimationCurve;
+
+	/** How many pixels of padding around the outside of the root-level tree node box.  Adding padding makes the tree look better, but you lose some sizing accuracy */
+	float TopLevelContainerOuterPadding;
+
+	/** How many pixels of padding around the outside of a non-root tree node's box.  Adding padding makes the tree look better, but you lose some sizing accuracy */
+	float NestedContainerOuterPadding;
+
+	/** How many pixels of spacing between the container and its child containers.  Adding padding makes the tree look better, but you lose some sizing accuracy */
+	float ContainerInnerPadding;
+
+	/** How many pixels to pad text that's drawn inside of a container (not the top-level container, though) */
+	float ChildContainerTextPadding;
+
+
+	/**
+	 * Live editing
+	 */
+
+	/** True if we should allow real-time editing right now.  The user can still navigate between nodes even if editing is disbled, so its not exactly "read only" */
+	TAttribute<bool> AllowEditing;
+
+	/** Node that we're currently dragging around.  This points to an entry in the CachedNodeVisuals array! */
+	FTreeMapNodeVisualInfo* DraggingVisual;
+
+	/** Distance we've LMB+dragged the cursor.  Used to determine if we're ready to start dragging and dropping */
+	float DragVisualDistance;
+
+	/** Position of the mouse cursor relative to the widget, where the user picked up a dragged visual */
+	FVector2D RelativeDragStartMouseCursorPos;
+
+	/** Position of the mouse cursor relative to the widget, the last time it moved.  Used for drag and drop. */
+	FVector2D RelativeMouseCursorPos;
+
+	/** Undo history buffer. */
+	TArray< TSharedPtr< FTreeMapNodeData > > UndoStates;
+
+	/** Current undo level, or INDEX_NONE if we haven't undone anything */
+	int32 CurrentUndoLevel;
+
+	/** Weak pointer to the widget used to rename a node inline.  We might need to close the window in some cases. */
+	TWeakPtr<SWidget> RenamingNodeWidget;
+
+	/** Weak pointer to the node that we're currently renaming, if any */
+	TWeakPtr<FTreeMapNodeData> RenamingNodeData;
+
+	/** True if the node we're renaming is a newly-created node that should be added to the tree after we finish naming it */
+	bool bIsNamingNewNode;
+
+	/** Node that we're currently drawing a highlight pulse effect for.  This points to an entry in the CachedNodesVisuals array */
+	TWeakPtr<FTreeMapNodeData> HighlightPulseNode;
+
+	/** Time that we started playing a highlight pulse effect for the HighlightPulseNode */
+	double HighlightPulseStartTime;
+
+};
diff --git a/Engine/Source/Developer/TreeMap/Public/TreeMapStyle.h b/Engine/Source/Developer/TreeMap/Public/TreeMapStyle.h
new file mode 100644
--- /dev/null
+++ b/Engine/Source/Developer/TreeMap/Public/TreeMapStyle.h
@@ -0,0 +1,24 @@
+// Copyright 1998-2013 Epic Games, Inc. All Rights Reserved.
+
+#pragma once
+
+
+/**
+ * Style data for STreeMap
+ */
+class TREEMAP_API FTreeMapStyle
+{
+public:
+	static void Initialize();
+	static void Shutdown();
+	static const ISlateStyle& Get();
+	static FName GetStyleSetName();
+
+private:
+
+	static TSharedRef< class FSlateStyleSet > Create();
+
+private:
+
+	static TSharedPtr< class FSlateStyleSet > StyleInstance;
+};
diff --git a/Engine/Source/Developer/TreeMap/STreeMap.cpp b/Engine/Source/Developer/TreeMap/STreeMap.cpp
new file mode 100644
--- /dev/null
+++ b/Engine/Source/Developer/TreeMap/STreeMap.cpp
@@ -0,0 +1,1554 @@
+// Copyright 1998-2013 Epic Games, Inc. All Rights Reserved.
+
+#include "TreeMapModule.h"
+#include "STreeMap.h"
+#include "ITreeMap.h"
+
+
+#define LOCTEXT_NAMESPACE "STreeMap"
+
+
+namespace STreeMapDefs
+{
+	/** Minimum pixels the user must drag the cursor before a drag+drop starts on a visual */
+	const float MinCursorDistanceForDraggingVisual = 3.0f;
+}
+
+
+
+void STreeMap::Construct( const FArguments& InArgs, const TSharedRef<FTreeMapNodeData>& InTreeMapNodeData, const TSharedPtr< ITreeMapCustomization >& InCustomization )
+{
+	CurrentUndoLevel = INDEX_NONE;
+	Customization = InCustomization;
+
+	TreeMapNodeData = InTreeMapNodeData;
+
+	AllowEditing = InArgs._AllowEditing;
+
+	BackgroundImage = InArgs._BackgroundImage;
+	NodeBackground = InArgs._NodeBackground;
+	HoveredNodeBackground = InArgs._HoveredNodeBackground;
+	BorderPadding = InArgs._BorderPadding;
+	RelativeDragStartMouseCursorPos = FVector2D::ZeroVector;
+	RelativeMouseCursorPos = FVector2D::ZeroVector;
+
+	{
+		FTreeMapOptions TreeMapOptions;
+		NameFont = TreeMapOptions.NameFont;
+		Name2Font = TreeMapOptions.Name2Font;
+		CenterTextFont = TreeMapOptions.CenterTextFont;
+
+		if( InArgs._NameFont.IsSet() )
+		{
+			NameFont = InArgs._NameFont;
+		}
+		if( InArgs._Name2Font.IsSet() )
+		{
+			Name2Font = InArgs._Name2Font;
+		}
+		if( InArgs._CenterTextFont.IsSet() )
+		{
+			CenterTextFont = InArgs._CenterTextFont;
+		}
+	}
+
+	MinimumVisualTreeNodeSize = InArgs._MinimumVisualTreeNodeSize;
+	TopLevelContainerOuterPadding = InArgs._TopLevelContainerOuterPadding;
+	NestedContainerOuterPadding = InArgs._NestedContainerOuterPadding;
+	ContainerInnerPadding = InArgs._ContainerInnerPadding;
+	ChildContainerTextPadding = InArgs._ChildContainerTextPadding;
+
+	TreeMapSize = FVector2D( 0, 0 );
+	NavigateAnimationCurve.AddCurve( 0.0f, InArgs._NavigationTransitionTime, ECurveEaseFunction::CubicOut );
+	MouseOverVisual = NULL;
+	DraggingVisual = NULL;
+	DragVisualDistance = 0.0f;
+
+	bIsNamingNewNode = false;
+
+	HighlightPulseStartTime = -99999.0;
+
+	OnTreeMapNodeDoubleClicked = InArgs._OnTreeMapNodeDoubleClicked;
+
+	if( Customization.IsValid() )
+	{
+		SizeNodesByAttribute = Customization->GetDefaultSizeByAttribute();
+		ColorNodesByAttribute = Customization->GetDefaultColorByAttribute();
+	}
+
+	const bool bShouldPlayTransition = false;
+	SetTreeRoot( TreeMapNodeData.ToSharedRef(), bShouldPlayTransition );
+}
+
+
+void STreeMap::Tick( const FGeometry& AllottedGeometry, const double InCurrentTime, const float InDeltaTime )
+{
+	SLeafWidget::Tick( AllottedGeometry, InCurrentTime, InDeltaTime );
+
+	if( AllottedGeometry.Size != TreeMapSize || !TreeMap.IsValid() )
+	{
+		// Stop renaming a node if we were doing that
+		StopRenamingNode();
+
+		// Make a tree map!
+		FTreeMapOptions TreeMapOptions;
+		TreeMapOptions.TreeMapType = ETreeMapType::Squarified;
+		TreeMapOptions.DisplayWidth = AllottedGeometry.Size.X;
+		TreeMapOptions.DisplayHeight = AllottedGeometry.Size.Y;
+		TreeMapOptions.TopLevelContainerOuterPadding = TopLevelContainerOuterPadding;
+		TreeMapOptions.NestedContainerOuterPadding = NestedContainerOuterPadding;
+		TreeMapOptions.ContainerInnerPadding = ContainerInnerPadding;
+		TreeMapOptions.NameFont = NameFont.Get();
+		TreeMapOptions.Name2Font = Name2Font.Get();
+		TreeMapOptions.CenterTextFont = CenterTextFont.Get();
+		TreeMapOptions.FontSizeChangeBasedOnDepth = 2;		// @todo treemap custom: Expose as a customization option to STreeMap
+		TreeMapOptions.MinimumInteractiveNodeSize = MinimumVisualTreeNodeSize;
+		TreeMap = ITreeMap::CreateTreeMap( TreeMapOptions, ActiveRootTreeMapNode.ToSharedRef() );
+
+		TreeMapSize = AllottedGeometry.Size;
+
+		CachedNodeVisuals = TreeMap->GetVisuals();
+		MouseOverVisual = NULL;
+		DraggingVisual = NULL;
+
+		// Map the new visuals back to the old ones!
+		NodeVisualIndicesToLastIndices.Reset();
+		TSet<int32> ValidLastIndices;
+		for( auto VisualIndex = 0; VisualIndex < CachedNodeVisuals.Num(); ++VisualIndex )
+		{
+			const auto& Visual = CachedNodeVisuals[ VisualIndex ];
+			for( auto LastVisualIndex = 0; LastVisualIndex < LastCachedNodeVisuals.Num(); ++LastVisualIndex )
+			{
+				const auto& LastVisual = LastCachedNodeVisuals[ LastVisualIndex ];
+				if( LastVisual.NodeData == Visual.NodeData )
+				{
+					NodeVisualIndicesToLastIndices.Add( VisualIndex, LastVisualIndex );
+					ValidLastIndices.Add( LastVisualIndex );
+					break;
+				}
+			}
+		}
+
+		// Find all of the orphans
+		OrphanedLastIndices.Reset();
+		for( auto LastVisualIndex = 0; LastVisualIndex < LastCachedNodeVisuals.Num(); ++LastVisualIndex )
+		{
+			if( !ValidLastIndices.Contains( LastVisualIndex ) )
+			{
+				OrphanedLastIndices.Add( LastVisualIndex );
+			}
+		}
+	}
+}
+
+
+void STreeMap::MakeBlendedNodeVisual( const int32 VisualIndex, const float NavigationAlpha, FTreeMapNodeVisualInfo& OutVisual ) const
+{
+	OutVisual = CachedNodeVisuals[ VisualIndex ];	// NOTE: Copying visual
+
+	// Do we need to interp?
+	if( NavigationAlpha < 1.0f )
+	{
+		// Did the visual exist before we navigated?
+		const int32* LastVisualIndexPtr = NodeVisualIndicesToLastIndices.Find( VisualIndex );
+		if( LastVisualIndexPtr != NULL )
+		{
+			// It did exist!
+			const auto LastVisualIndex = *LastVisualIndexPtr;
+			const auto& LastVisual = LastCachedNodeVisuals[ LastVisualIndex ];
+
+			// Blend before "before" and "now"
+			OutVisual.Position = FMath::Lerp( LastVisual.Position, OutVisual.Position, NavigationAlpha );
+			OutVisual.Size = FMath::Lerp( LastVisual.Size, OutVisual.Size, NavigationAlpha );
+
+			// Do an HSV color lerp; it just looks more sensible!
+			OutVisual.Color = FLinearColor::LerpUsingHSV( LastVisual.Color, OutVisual.Color, NavigationAlpha );
+
+			// The blended visual is considered interactive only if both the new version and the old version were interactive
+			OutVisual.bIsInteractive = LastVisual.bIsInteractive && OutVisual.bIsInteractive;
+		}
+		else
+		{
+			// Didn't exist before.  Fade in from nothing!
+			OutVisual.Color.A *= NavigationAlpha;
+		}
+	}
+}
+
+
+
+int32 STreeMap::OnPaint( const FPaintArgs& Args, const FGeometry& AllottedGeometry, const FSlateRect& WidgetClippingRect, FSlateWindowElementList& OutDrawElements, int32 LayerId, const FWidgetStyle& InWidgetStyle, bool bParentEnabled ) const
+{
+	const bool bEnabled = ShouldBeEnabled( bParentEnabled );
+	const auto DrawEffects = bEnabled ? ESlateDrawEffect::None : ESlateDrawEffect::DisabledEffect;
+
+	// Draw background border layer
+	{
+		const FSlateBrush* ThisBackgroundImage = BackgroundImage.Get();
+		FSlateDrawElement::MakeBox(
+			OutDrawElements,
+			LayerId,
+			AllottedGeometry.ToPaintGeometry(),
+			ThisBackgroundImage,
+			WidgetClippingRect,
+			DrawEffects,
+			InWidgetStyle.GetColorAndOpacityTint() * ThisBackgroundImage->TintColor.GetColor( InWidgetStyle )
+			);
+	}
+
+	float NavigationAlpha = NavigateAnimationCurve.GetLerp();
+
+
+	// Draw tree map
+	if( TreeMap.IsValid() )
+	{
+		// Figure out all the nodes we need to draw in a big ordered list.  This can sometimes include nodes from the previous tree map we were
+		// looking at before we "navigated", as those nodes will be briefly visible during the animated transition
+		static TArray< const FTreeMapNodeVisualInfo* > NodeVisualsToDraw;
+		NodeVisualsToDraw.Reset();
+
+		// Draw the orphan's first
+		// @todo treemap visuals: During transitions, nodes look "on top" should be drawn last.  Right now they will underlap and overlap adjacents at the same time.  Looks weird.
+		for( auto OrphanedVisualIndex : OrphanedLastIndices )
+		{
+			NodeVisualsToDraw.Add( &LastCachedNodeVisuals[ OrphanedVisualIndex ] );
+		}
+		const int32 FirstNewVisualIndex = NodeVisualsToDraw.Num();
+		for( const auto& Visual : CachedNodeVisuals )
+		{
+			NodeVisualsToDraw.Add( &Visual );
+		}
+
+		const FTreeMapNodeData* RenamingNodeDataPtr = RenamingNodeData.Pin().Get();
+
+		// Draw background boxes for all visuals
+		{
+			++LayerId;
+
+			const FSlateBrush* ThisHoveredNodeBackground = HoveredNodeBackground.Get();
+			const FSlateBrush* ThisNodeBackground = NodeBackground.Get();
+
+			for( auto DrawVisualIndex = 0; DrawVisualIndex < NodeVisualsToDraw.Num(); ++DrawVisualIndex )
+			{
+				FTreeMapNodeVisualInfo BlendedVisual;
+				if( DrawVisualIndex >= FirstNewVisualIndex )
+				{
+					MakeBlendedNodeVisual( DrawVisualIndex - FirstNewVisualIndex, NavigationAlpha, BlendedVisual );
+				}
+				else
+				{
+					// This is an orphan
+					BlendedVisual = *NodeVisualsToDraw[ DrawVisualIndex ];
+					BlendedVisual.Color.A *= 1.0f - NavigationAlpha;	// Fade orphans out when navigating
+				}
+
+				// Don't draw if completely faded out
+				if( BlendedVisual.Color.A > KINDA_SMALL_NUMBER )
+				{
+					const bool bIsMouseOverNode = MouseOverVisual != NULL && MouseOverVisual->NodeData == BlendedVisual.NodeData;
+
+					// Draw the visual's background box
+					const FVector2D VisualPosition = BlendedVisual.Position;
+					const FSlateRect VisualClippingRect = TransformRect( AllottedGeometry.GetAccumulatedLayoutTransform(), FSlateRect( VisualPosition, VisualPosition + BlendedVisual.Size ) );
+					const auto VisualPaintGeometry = AllottedGeometry.ToPaintGeometry( VisualPosition, BlendedVisual.Size );
+					auto DrawColor = InWidgetStyle.GetColorAndOpacityTint() * ThisNodeBackground->TintColor.GetColor( InWidgetStyle ) * BlendedVisual.Color;
+					FSlateDrawElement::MakeBox(
+						OutDrawElements,
+						LayerId,
+						VisualPaintGeometry,
+						bIsMouseOverNode ? ThisHoveredNodeBackground : ThisNodeBackground,
+						VisualClippingRect,
+						DrawEffects,
+						DrawColor
+						);
+
+					if( BlendedVisual.NodeData->BackgroundBrush != nullptr )
+					{
+						// Preserve aspect ratio, but stretch to fill the whole rectangle, even if we have to crop the smaller edge
+						const float LargestSize = FMath::Max( BlendedVisual.Size.X, BlendedVisual.Size.Y );
+						const FVector2D BackgroundSize( LargestSize, LargestSize );
+						FVector2D BackgroundPosition = VisualPosition;
+						if( BlendedVisual.Size.X > BlendedVisual.Size.Y )	// Is our box wider than tall?
+						{
+							const float SizeDifference = LargestSize - BlendedVisual.Size.Y;
+							BackgroundPosition.Y -= SizeDifference * 0.5f;
+						}
+						else
+						{
+							const float SizeDifference = LargestSize - BlendedVisual.Size.X;
+							BackgroundPosition.X -= SizeDifference * 0.5f;
+						}
+						const auto BackgroundPaintGeometry = AllottedGeometry.ToPaintGeometry( BackgroundPosition, BackgroundSize );
+
+						const FSlateRect BackgroundClippingRect = VisualClippingRect.InsetBy( FMargin( 1 ) );
+
+						// Draw the background brush
+						FSlateDrawElement::MakeBox(
+							OutDrawElements,
+							LayerId,
+							BackgroundPaintGeometry,
+							BlendedVisual.NodeData->BackgroundBrush,
+							BackgroundClippingRect,
+							DrawEffects,
+							DrawColor );
+					}
+
+					const bool bIsHighlightPulseNode = HighlightPulseNode.IsValid() && HighlightPulseNode.Pin().Get() == BlendedVisual.NodeData;
+					if( bIsHighlightPulseNode )
+					{
+						const float HighlightPulseAnimDuration = 1.5f;	// @todo treemap: Probably should be customizable in the widget's settings
+						const float HighlightPulseAnimProgress = ( FSlateApplication::Get().GetCurrentTime() - HighlightPulseStartTime ) / HighlightPulseAnimDuration;
+						if( HighlightPulseAnimProgress >= 0.0f && HighlightPulseAnimProgress <= 1.0f )
+						{
+							DrawColor = FLinearColor( 1.0f, 1.0f, 1.0f, FMath::MakePulsatingValue( HighlightPulseAnimProgress, 6.0f, 0.5f ) );
+							FSlateDrawElement::MakeBox(
+								OutDrawElements,
+								LayerId,
+								VisualPaintGeometry,
+								ThisHoveredNodeBackground,
+								VisualClippingRect,
+								DrawEffects,
+								DrawColor );
+						}
+					}
+
+				}
+			}
+		}
+
+
+		// Draw text layers.  We draw it twice for all visuals.  Once for a drop shadow, and then again for the foreground text.
+		const TSharedRef< FSlateFontMeasure >& FontMeasureService = FSlateApplication::Get().GetRenderer()->GetFontMeasureService();
+		for( auto TextLayerIndex = 0; TextLayerIndex < 2; ++TextLayerIndex )
+		{
+			++LayerId;
+			
+			// @todo treemap: Want ellipses for truncated text
+			// @todo treemap: Squash text based on box size rather than strictly depth?
+
+			const float ShadowOpacity = 0.5f;
+			const auto ShadowOffset = ( TextLayerIndex == 0 ) ? FVector2D( -1.0f, 1.0f ) : FVector2D::ZeroVector;
+			const auto TextColorScale = ( TextLayerIndex == 0 ) ? FLinearColor( 0.0f, 0.0f, 0.0f, ShadowOpacity ) : FLinearColor::White;
+
+			for( auto DrawVisualIndex = 0; DrawVisualIndex < NodeVisualsToDraw.Num(); ++DrawVisualIndex )
+			{
+				FTreeMapNodeVisualInfo BlendedVisual;
+				if( DrawVisualIndex >= FirstNewVisualIndex )
+				{
+					MakeBlendedNodeVisual( DrawVisualIndex - FirstNewVisualIndex, NavigationAlpha, BlendedVisual );
+				}
+				else
+				{
+					// This is an orphan
+					BlendedVisual = *NodeVisualsToDraw[ DrawVisualIndex ];
+					BlendedVisual.Color.A *= 1.0f - NavigationAlpha;	// Fade orphans out when navigating
+				}
+
+				// Don't draw text if completely faded out, or if not interactive
+				if( BlendedVisual.Color.A > KINDA_SMALL_NUMBER && BlendedVisual.bIsInteractive )
+				{
+					// Don't draw a title for a node that is actively being renamed -- the text box is already visible, after all
+					if( BlendedVisual.NodeData != RenamingNodeDataPtr )
+					{
+						const FVector2D VisualPosition = BlendedVisual.Position;
+
+						// Allow text to fade out with its parent box
+						auto VisualTextColor = TextColorScale;
+						VisualTextColor.A *= BlendedVisual.Color.A;
+
+						const bool bIsMouseOverNode = MouseOverVisual != NULL && MouseOverVisual->NodeData == BlendedVisual.NodeData;
+
+						// If the visual is crushed down pretty small in screen space, we don't want to bother even try drawing text
+						const FVector2D ScreenSpaceVisualSize( AllottedGeometry.Scale * BlendedVisual.Size );
+						if( ScreenSpaceVisualSize.X > 20 )
+						{
+							const auto VisualPaintGeometry = AllottedGeometry.ToPaintGeometry( VisualPosition, BlendedVisual.Size );
+
+							// Clip the text to the visual's rectangle, with some extra inner padding to avoid overlapping the visual's border
+							auto TextClippingRect = FSlateRect( VisualPaintGeometry.DrawPosition, VisualPaintGeometry.DrawPosition + VisualPaintGeometry.GetLocalSize() );
+							TextClippingRect = TextClippingRect.IntersectionWith( WidgetClippingRect );
+							TextClippingRect = TextClippingRect.InsetBy( FMargin( ChildContainerTextPadding, 0, ChildContainerTextPadding, 0 ) );
+							if( TextClippingRect.IsValid() )
+							{
+								// Name (first line)
+								float NameTextHeight = 0.0f;
+								if( !BlendedVisual.NodeData->Name.IsEmpty() )
+								{
+									const FVector2D TextSize = FontMeasureService->Measure( BlendedVisual.NodeData->Name, BlendedVisual.NameFont );
+									NameTextHeight = TextSize.Y;
+									const auto TextX = VisualPosition.X + FMath::Max( ChildContainerTextPadding, (float)BlendedVisual.Size.X * 0.5f - TextSize.X * 0.5f );	// Clamp to left edge if cropped, so the user can at least read the beginning
+									const auto TextY = VisualPosition.Y + ChildContainerTextPadding;
+
+									FSlateDrawElement::MakeText(
+										OutDrawElements,
+										LayerId,
+										AllottedGeometry.ToOffsetPaintGeometry( ShadowOffset + FVector2D( TextX, TextY ) ),
+										BlendedVisual.NodeData->Name,
+										BlendedVisual.NameFont,
+										TextClippingRect,
+										DrawEffects,
+										InWidgetStyle.GetColorAndOpacityTint() * VisualTextColor );
+								}
+
+								// Name (second line)
+								float Name2TextHeight = 0.0f;
+								if( !BlendedVisual.NodeData->Name2.IsEmpty() && BlendedVisual.NodeData->IsLeafNode() )
+								{
+									const FVector2D TextSize = FontMeasureService->Measure( BlendedVisual.NodeData->Name2, BlendedVisual.Name2Font );
+									Name2TextHeight = TextSize.Y;
+									const auto TextX = VisualPosition.X + FMath::Max( ChildContainerTextPadding, (float)BlendedVisual.Size.X * 0.5f - TextSize.X * 0.5f );	// Clamp to left edge if cropped, so the user can at least read the beginning
+									const auto TextY = VisualPosition.Y + ChildContainerTextPadding + NameTextHeight;
+
+									// Clip the text to the visual's rectangle, with some extra inner padding to avoid overlapping the visual's border
+									FSlateDrawElement::MakeText(
+										OutDrawElements,
+										LayerId,
+										AllottedGeometry.ToOffsetPaintGeometry( ShadowOffset + FVector2D( TextX, TextY ) ),
+										BlendedVisual.NodeData->Name2,
+										BlendedVisual.Name2Font,
+										TextClippingRect,
+										DrawEffects,
+										InWidgetStyle.GetColorAndOpacityTint() * VisualTextColor );
+								}
+
+								// Center text
+								if( !BlendedVisual.NodeData->CenterText.IsEmpty() && BlendedVisual.NodeData->IsLeafNode() )
+								{
+									// If the visual is smaller than the text we're going to draw, try using a smaller font
+									const FSlateFontInfo* CenterTextFont = &BlendedVisual.CenterTextFont;
+									const FVector2D FullTextSize = FontMeasureService->Measure( BlendedVisual.NodeData->CenterText, *CenterTextFont );
+									if( ScreenSpaceVisualSize.X < FullTextSize.X || ScreenSpaceVisualSize.Y < FullTextSize.Y )	// @todo treemap: assumption that NameFont will be smaller than CenterText font
+									{
+										CenterTextFont = &BlendedVisual.NameFont;
+									}
+
+									const FVector2D TextSize = FontMeasureService->Measure( BlendedVisual.NodeData->CenterText, *CenterTextFont );
+									const auto TextX = VisualPosition.X + FMath::Max( ChildContainerTextPadding, (float)BlendedVisual.Size.X * 0.5f - TextSize.X * 0.5f );	// Clamp to left edge if cropped, so the user can at least read the beginning
+									const auto TextY = VisualPosition.Y + FMath::Max( ChildContainerTextPadding + NameTextHeight + Name2TextHeight, (float)BlendedVisual.Size.Y * 0.5f - TextSize.Y * 0.5f );	// Clamp to bottom of first line of text if cropped
+
+									// Clip the text to the visual's rectangle, with some extra inner padding to avoid overlapping the visual's border
+									FSlateDrawElement::MakeText(
+										OutDrawElements,
+										LayerId,
+										AllottedGeometry.ToOffsetPaintGeometry( ShadowOffset + FVector2D( TextX, TextY ) ),
+										BlendedVisual.NodeData->CenterText,
+										*CenterTextFont,
+										TextClippingRect,
+										DrawEffects,
+										InWidgetStyle.GetColorAndOpacityTint() * VisualTextColor );
+								}
+							}
+						}
+					}
+				}
+			}
+		}
+	}
+
+
+	if( DraggingVisual != NULL && DragVisualDistance >= STreeMapDefs::MinCursorDistanceForDraggingVisual )
+	{
+		const FVector2D DragStartCursorPos = RelativeDragStartMouseCursorPos;
+		const FVector2D NewCursorPos = RelativeMouseCursorPos;
+
+		const FVector2D SplineStart = DragStartCursorPos;
+		const FVector2D SplineEnd = NewCursorPos;
+
+		const FVector2D SplineStartDir = ( DragStartCursorPos - ( DraggingVisual->Position + DraggingVisual->Size * 0.5f ) );	// @todo treemap: Point away from the center of the dragged visual
+		const FVector2D SplineEndDir = FVector2D( 0.0f, 200.0f );	// @todo treemap: Probably needs better customization support
+
+		// @todo treemap: Draw line in red if drop won't work?
+
+		// Draw two passes, the first one is an drop shadow
+		for( auto SplineLayerIndex = 0; SplineLayerIndex < 2; ++SplineLayerIndex )
+		{
+			++LayerId;
+
+			const float ShadowOpacity = 0.5f;
+			const float SplineThickness = ( SplineLayerIndex == 0 ) ? 5.0f : 4.0f;
+			const auto ShadowOffset = ( SplineLayerIndex == 0 ) ? FVector2D( -1.0f, 1.0f ) : FVector2D::ZeroVector;
+			const auto SplineColorScale = ( SplineLayerIndex == 0 ) ? FLinearColor( 0.0f, 0.0f, 0.0f, ShadowOpacity ) : FLinearColor::White;
+
+			FSlateDrawElement::MakeSpline(
+				OutDrawElements,
+				LayerId,
+				AllottedGeometry.ToPaintGeometry(  ShadowOffset, FVector2D( 1.0, 1.0f ) ),
+				SplineStart,
+				SplineStartDir,
+				SplineEnd,
+				SplineEndDir,
+				WidgetClippingRect,
+				SplineThickness,
+				DrawEffects,
+				InWidgetStyle.GetColorAndOpacityTint() * DraggingVisual->Color * SplineColorScale );
+		}
+	}
+
+
+	return LayerId;
+}
+
+
+
+FVector2D STreeMap::ComputeDesiredSize(float LayoutScaleMultiplier) const
+{
+	// TreeMap widgets have no desired size -- their size is always determined by their container
+	return FVector2D::ZeroVector;
+}
+
+
+FReply STreeMap::OnMouseMove( const FGeometry& InMyGeometry, const FPointerEvent& InMouseEvent )
+{
+	FReply Reply = FReply::Unhandled();
+
+	// Update window-relative cursor position
+	RelativeMouseCursorPos = InMyGeometry.AbsoluteToLocal( InMouseEvent.GetScreenSpacePosition() );
+
+	// Clear current hover
+	MouseOverVisual = NULL;
+
+	// Don't hover while transitioning.  It looks bad!
+	if( !IsNavigationTransitionActive() )
+	{
+		// Figure out which visual the cursor is over
+		FTreeMapNodeVisualInfo* NodeVisualUnderCursor = FindNodeVisualUnderCursor( InMyGeometry, InMouseEvent.GetScreenSpacePosition() );
+		if( NodeVisualUnderCursor != NULL )
+		{
+			// Mouse is over a visual
+			MouseOverVisual = NodeVisualUnderCursor;
+		}
+
+		Reply = FReply::Handled();
+	}
+
+	if( DraggingVisual != NULL )
+	{
+		DragVisualDistance += InMouseEvent.GetCursorDelta().Size();
+	}
+
+	return Reply;
+}
+
+
+void STreeMap::OnMouseLeave( const FPointerEvent& InMouseEvent )
+{
+	MouseOverVisual = NULL;
+}
+
+
+FReply STreeMap::OnMouseButtonDown( const FGeometry& InMyGeometry, const FPointerEvent& InMouseEvent )
+{
+	FReply Reply = FReply::Unhandled();
+	if( InMouseEvent.GetEffectingButton() == EKeys::LeftMouseButton )
+	{
+		// Wait until we've finished animating a transition
+		if( !IsNavigationTransitionActive() )
+		{
+			if( AllowEditing.Get() )
+			{
+				// Figure out what was clicked on
+				FTreeMapNodeVisualInfo* NodeVisualUnderCursor = FindNodeVisualUnderCursor( InMyGeometry, InMouseEvent.GetScreenSpacePosition() );
+				if( NodeVisualUnderCursor != NULL )
+				{
+					// Mouse was pressed on a node
+					DraggingVisual = NodeVisualUnderCursor;
+					DragVisualDistance = 0.0f;
+					RelativeDragStartMouseCursorPos = InMyGeometry.AbsoluteToLocal( InMouseEvent.GetScreenSpacePosition() );
+
+					Reply = FReply::Handled();
+				}
+			}
+		}
+	}
+
+	return Reply;
+}
+
+
+FReply STreeMap::OnMouseButtonUp( const FGeometry& InMyGeometry, const FPointerEvent& InMouseEvent )
+{
+	FReply Reply = FReply::Unhandled();
+	if( InMouseEvent.GetEffectingButton() == EKeys::LeftMouseButton )
+	{
+		auto* DroppedVisual = DraggingVisual;
+		DraggingVisual = NULL;
+
+		if( DroppedVisual != NULL && DragVisualDistance >= STreeMapDefs::MinCursorDistanceForDraggingVisual )
+		{
+			// Wait until we've finished animating a transition
+			if( !IsNavigationTransitionActive() )
+			{
+				// Find the node under the cursor
+				FTreeMapNodeVisualInfo* NodeVisualUnderCursor = FindNodeVisualUnderCursor( InMyGeometry, InMouseEvent.GetScreenSpacePosition() );
+				if( NodeVisualUnderCursor != NULL )
+				{
+					// The dropped node will become a child of the node we dropped onto
+					auto NewParentNode = NodeVisualUnderCursor->NodeData->AsShared();
+					auto DroppedNode = DroppedVisual->NodeData->AsShared();
+
+					// Reparent it!
+					ReparentNode( DroppedNode, NewParentNode );
+
+					Reply = FReply::Handled();
+				}
+			}
+		}
+		else
+		{
+			FTreeMapNodeVisualInfo* NodeVisualUnderCursor = FindNodeVisualUnderCursor( InMyGeometry, InMouseEvent.GetScreenSpacePosition() );
+			if( NodeVisualUnderCursor != NULL )
+			{
+				if( AllowEditing.Get() )
+				{
+					// Start renaming!
+					const bool bIsNewNode = false;
+					StartRenamingNode( InMyGeometry, NodeVisualUnderCursor->NodeData->AsShared(), NodeVisualUnderCursor->Position, bIsNewNode );
+
+					Reply = FReply::Handled();
+				}
+			}
+		}
+
+		DragVisualDistance = 0.0f;
+	}
+	else if( InMouseEvent.GetEffectingButton() == EKeys::RightMouseButton )
+	{
+		// Show a pop-up menu!
+		ShowOptionsMenuAt( InMouseEvent.GetScreenSpacePosition() );
+	}
+
+	return Reply;
+}
+
+
+
+FReply STreeMap::OnMouseButtonDoubleClick( const FGeometry& InMyGeometry, const FPointerEvent& InMouseEvent )
+{
+	FReply Reply = FReply::Unhandled();
+
+	if( InMouseEvent.GetEffectingButton() == EKeys::LeftMouseButton )
+	{
+		// Wait until we're done transitioning before allowing another transition
+		if( !IsNavigationTransitionActive() )
+		{
+			// Figure out what was clicked on
+			const FTreeMapNodeVisualInfo* NodeVisualUnderCursor = FindNodeVisualUnderCursor( InMyGeometry, InMouseEvent.GetScreenSpacePosition() );
+			if( NodeVisualUnderCursor != NULL )
+			{
+				// Double-clicked on a tree map visual!  Check to see if we were asked to customize how double-click is handled.
+				if( OnTreeMapNodeDoubleClicked.IsBound() )
+				{
+					OnTreeMapNodeDoubleClicked.Execute( *NodeVisualUnderCursor->NodeData );
+				}
+				else
+				{
+					// Double-click was not overridden, so just do our default thing and re-root the tree directly on the node that was double-clicked on
+					
+					// Re-root the tree
+					const bool bShouldPlayTransition = true;
+					SetTreeRoot( NodeVisualUnderCursor->NodeData->AsShared(), bShouldPlayTransition );
+				}
+
+				Reply = FReply::Handled();
+			}
+		}
+	}
+
+	return Reply;
+}
+
+
+FReply STreeMap::OnMouseWheel( const FGeometry& MyGeometry, const FPointerEvent& MouseEvent )
+{
+	FReply Reply = FReply::Unhandled();
+
+	// Don't zoom in or out while already transitioning.  It feels too frenetic.
+	if( !IsNavigationTransitionActive() )
+	{
+		if( MouseEvent.GetWheelDelta() > 0 )
+		{
+			// Figure out what was clicked on
+			const FTreeMapNodeVisualInfo* NodeVisualUnderCursor = FindNodeVisualUnderCursor( MyGeometry, MouseEvent.GetScreenSpacePosition() );
+			if( NodeVisualUnderCursor != NULL )
+			{
+				// From the node that was scrolled over, visits nodes upward until we find one whose parent is our current active root
+				FTreeMapNodeDataPtr NextNode = NodeVisualUnderCursor->NodeData->AsShared();
+				do
+				{
+					FTreeMapNodeDataPtr NodeParent;
+					if( NextNode->Parent != NULL )
+					{
+						NodeParent = NextNode->Parent->AsShared();
+					}
+
+					if( NodeParent == ActiveRootTreeMapNode )
+					{
+						break;
+					}
+					NextNode = NodeParent;
+				}
+				while( NextNode.IsValid() );
+
+				// Zoom in one level
+				if( NextNode.IsValid() )
+				{
+					const bool bShouldPlayTransition = true;
+					SetTreeRoot( NextNode.ToSharedRef(), bShouldPlayTransition );
+
+					if( !Reply.IsEventHandled() )
+					{
+						Reply = FReply::Handled();
+					}
+				}
+			}
+		}
+		else if( MouseEvent.GetWheelDelta() < 0 )
+		{
+			// Zoom out one level
+			if( ActiveRootTreeMapNode->Parent != NULL )
+			{
+				const bool bShouldPlayTransition = true;
+				SetTreeRoot( ActiveRootTreeMapNode->Parent->AsShared(), bShouldPlayTransition );
+
+				if( !Reply.IsEventHandled() )
+				{
+					Reply = FReply::Handled();
+				}
+			}
+		}
+	}
+
+	return Reply;
+}
+
+
+void STreeMap::SetTreeRoot( const FTreeMapNodeDataRef& NewRoot, const bool bShouldPlayTransition )
+{
+	if( ActiveRootTreeMapNode != NewRoot )
+	{
+		ActiveRootTreeMapNode = NewRoot;
+
+		// Freshen the visualization data for the node since it may be stale after being copied off the undo stack
+		RebuildTreeMap( bShouldPlayTransition );
+	}
+}
+
+
+void STreeMap::RebuildTreeMap( const bool bShouldPlayTransition )
+{
+	// Stop renaming anything.  We don't want pop-up windows persisting during the transition
+	StopRenamingNode();
+
+	// Keep track of the last tree
+	LastTreeMap = TreeMap;
+	LastCachedNodeVisuals = CachedNodeVisuals;
+	if( bShouldPlayTransition )
+	{
+		NavigateAnimationCurve.Play( AsShared() );
+	}
+	else
+	{
+		NavigateAnimationCurve.JumpToEnd();
+	}
+
+	// Kill the tree map so that it will be regenerated
+	TreeMap.Reset();
+	CachedNodeVisuals.Reset();
+	DraggingVisual = NULL;
+	MouseOverVisual = NULL;
+
+	// Customize the size and coloring of the source nodes before we rebuild it
+	ApplyVisualizationToNodes( ActiveRootTreeMapNode.ToSharedRef() );
+}
+
+
+FTreeMapNodeVisualInfo* STreeMap::FindNodeVisualUnderCursor( const FGeometry& MyGeometry, const FVector2D& ScreenSpaceCursorPosition )
+{
+	if( TreeMap.IsValid() )
+	{
+		const FVector2D LocalCursorPosition = MyGeometry.AbsoluteToLocal( ScreenSpaceCursorPosition );
+
+		// NOTE: Iterating backwards so that child-most nodes are checked first
+		for( auto VisualIndex = CachedNodeVisuals.Num() - 1; VisualIndex >= 0; --VisualIndex )
+		{
+			auto& Visual = CachedNodeVisuals[ VisualIndex ];
+
+			// We only allow interactive visuals to be moused over
+			if( Visual.bIsInteractive )
+			{
+				const auto VisualRect = FSlateRect( Visual.Position, Visual.Position + Visual.Size );
+				if( VisualRect.ContainsPoint( LocalCursorPosition ) )
+				{
+					return &Visual;
+				}
+			}
+		}
+	}
+
+	return NULL;
+}
+
+
+bool STreeMap::IsNavigationTransitionActive() const
+{
+	return NavigateAnimationCurve.IsPlaying();
+}
+
+
+void STreeMap::TakeUndoSnapshot()
+{
+	// If we've already undone some state, then we'll remove any undo state beyond the level that
+	// we've already undone up to.
+	if( CurrentUndoLevel != INDEX_NONE )
+	{
+		UndoStates.RemoveAt( CurrentUndoLevel, UndoStates.Num() - CurrentUndoLevel );
+
+		// Reset undo level as we haven't undone anything since this latest undo
+		CurrentUndoLevel = INDEX_NONE;
+	}
+
+	// Take an Undo snapshot before we change anything
+	UndoStates.Add( this->TreeMapNodeData->CopyNodeRecursively() );
+
+	// If we've hit the undo limit, then delete previous entries
+	const int32 MaxUndoLevels = 200;	// @todo treemap custom: Make customizable in the settings of the widget?
+	if( UndoStates.Num() > MaxUndoLevels )
+	{
+		UndoStates.RemoveAt( 0 );
+	}
+}
+
+
+FTreeMapNodeDataPtr STreeMap::FindNodeInCopiedTree( const FTreeMapNodeDataRef& NodeToFind, const FTreeMapNodeDataRef& OriginalNode, const FTreeMapNodeDataRef& CopiedNode ) const
+{
+	if( NodeToFind == OriginalNode )
+	{
+		return CopiedNode;
+	}
+
+	for( int32 ChildNodeIndex = 0; ChildNodeIndex < OriginalNode->Children.Num(); ++ChildNodeIndex )
+	{
+		const auto& OriginalChildNode = OriginalNode->Children[ ChildNodeIndex ];
+		const auto& CopiedChildNode = CopiedNode->Children[ ChildNodeIndex ];
+
+		const auto& FoundMatchingCopy = FindNodeInCopiedTree( NodeToFind, OriginalChildNode.ToSharedRef(), CopiedChildNode.ToSharedRef() );
+		if( FoundMatchingCopy.IsValid() )
+		{
+			return FoundMatchingCopy;
+		}
+	}
+
+	return NULL;
+}
+
+
+void STreeMap::Undo()
+{
+	if( UndoStates.Num() > 0 )
+	{
+		// Restore from undo state
+		int32 UndoStateIndex;
+		if( CurrentUndoLevel == INDEX_NONE )
+		{
+			// We haven't undone anything since the last time a new undo state was added
+			UndoStateIndex = UndoStates.Num() - 1;
+
+			// Store an undo state for the current state (before undo was pressed)
+			TakeUndoSnapshot();
+		}
+		else
+		{
+			// Move down to the next undo level
+			UndoStateIndex = CurrentUndoLevel - 1;
+		}
+
+		// Is there anything else to undo?
+		if( UndoStateIndex >= 0 )
+		{
+			// Undo from history
+			{
+				const FTreeMapNodeDataRef NewTreeMapNodeData = UndoStates[ UndoStateIndex ]->CopyNodeRecursively();
+				const FTreeMapNodeDataRef NewActiveRoot = FindNodeInCopiedTree( ActiveRootTreeMapNode.ToSharedRef(), TreeMapNodeData.ToSharedRef(), NewTreeMapNodeData ).ToSharedRef();
+				TreeMapNodeData = NewTreeMapNodeData;
+
+				const bool bShouldPlayTransition = true;		// @todo treemap: A bit worried about volatility of LastTreeMap node refs here
+				SetTreeRoot( NewActiveRoot, bShouldPlayTransition );
+			}
+
+			CurrentUndoLevel = UndoStateIndex;
+		}
+	}
+}
+
+
+void STreeMap::Redo()
+{
+	// Is there anything to redo?  If we've haven't tried to undo since the last time new
+	// undo state was added, then CurrentUndoLevel will be INDEX_NONE
+	if( CurrentUndoLevel != INDEX_NONE )
+	{
+		const int32 NextUndoLevel = CurrentUndoLevel + 1;
+		if( UndoStates.Num() > NextUndoLevel )
+		{
+			// Restore from undo state
+			{
+				const FTreeMapNodeDataRef NewTreeMapNodeData = UndoStates[ NextUndoLevel ]->CopyNodeRecursively();
+				const FTreeMapNodeDataRef NewActiveRoot = FindNodeInCopiedTree( ActiveRootTreeMapNode.ToSharedRef(), TreeMapNodeData.ToSharedRef(), NewTreeMapNodeData ).ToSharedRef();
+				TreeMapNodeData = NewTreeMapNodeData;
+
+				const bool bShouldPlayTransition = true;	// @todo treemap: A bit worried about volatility of LastTreeMap node refs here
+				SetTreeRoot( NewActiveRoot, bShouldPlayTransition );
+			}
+
+			CurrentUndoLevel = NextUndoLevel;
+
+			if( UndoStates.Num() <= CurrentUndoLevel + 1 )
+			{
+				// We've redone all available undo states
+				CurrentUndoLevel = INDEX_NONE;
+
+				// Pop last undo state that we created on initial undo
+				UndoStates.RemoveAt( UndoStates.Num() - 1 );
+			}
+		}
+	}
+}
+
+
+void STreeMap::ReparentNode( const FTreeMapNodeDataRef DroppedNode, const FTreeMapNodeDataRef NewParentNode )
+{
+	// Can't reparent the tree root node
+	if( DroppedNode != TreeMapNodeData )
+	{
+		struct Local
+		{
+			/** Checks to see if the specified Node is a child of 'PossibleParent' at any level */
+			static bool IsMyParent( const FTreeMapNodeDataRef& Node, const FTreeMapNodeDataRef& PossibleParent )
+			{
+				if( Node == PossibleParent )
+				{
+					return true;
+				}
+
+				if( PossibleParent->Parent != NULL )
+				{
+					return IsMyParent( Node, PossibleParent->Parent->AsShared() );
+				}
+
+				return false;
+			}
+		};
+
+		// Can't parent owning nodes to their children
+		if( !Local::IsMyParent( DroppedNode, NewParentNode ) )
+		{
+			// Can't reparent to self
+			if( NewParentNode != DroppedNode )
+			{
+				// Only if parent has changed
+				if( &NewParentNode.Get() != DroppedNode->Parent )
+				{
+					// Take an Undo snapshot before we change anything
+					TakeUndoSnapshot();
+
+					if( DroppedNode->Parent != NULL )
+					{
+						DroppedNode->Parent->Children.RemoveSingle( DroppedNode );
+					}
+					NewParentNode->Children.Add( DroppedNode );
+					DroppedNode->Parent = &NewParentNode.Get();
+
+					// If we container node became a leaf node, we need to make sure it has a valid size set
+					// @todo treemap: This seems a bit... weird.  It won't reverse either, if you put the node back (save with overridden size!)  Maybe we need a bool for "auto size"=true.  OR, have a separate size for Node vs. Leaf?
+					if( DroppedNode->Size == 0.0f && DroppedNode->IsLeafNode() )
+					{
+						DroppedNode->Size = 1.0f;
+					}
+
+					// Invalidate the tree
+					const bool bShouldPlayTransition = true;
+					RebuildTreeMap( bShouldPlayTransition );
+
+					// After we have a new tree, play a highlight effect on the reparented node so the user can see where it ended
+					// up in the tree.  Tree map layout can be fairly volatile after a parenting change.
+					HighlightPulseNode = DroppedNode;
+					HighlightPulseStartTime = FSlateApplication::Get().GetCurrentTime();
+				}
+			}
+		}
+	}
+}
+
+
+FReply STreeMap::DeleteHoveredNode()
+{
+	FReply Reply = FReply::Unhandled();
+
+	// Wait until we've finished animating a transition
+	if( !IsNavigationTransitionActive() )
+	{
+		if( MouseOverVisual != NULL )
+		{
+			// Only non-root nodes can be deleted
+			auto NodeToDelete = MouseOverVisual->NodeData->AsShared();
+			if( NodeToDelete->Parent != NULL )
+			{
+				// Don't allow the current active tree root to be deleted.  The user should zoom out first!
+				if( NodeToDelete != ActiveRootTreeMapNode )
+				{
+					TakeUndoSnapshot();
+
+					// Delete the node
+					{
+						NodeToDelete->Parent->Children.RemoveSingle( NodeToDelete );
+						NodeToDelete->Parent = NULL;
+
+						// Note:  NodeToDelete will actually be deleted when it goes out of scope later in this function (shared pointer), but it is already removed from our tree
+					}
+
+					const bool bShouldPlayTransition = true;
+					RebuildTreeMap( bShouldPlayTransition );
+
+					Reply = FReply::Handled();
+				}
+			}
+		}
+	}
+
+	return Reply;
+}
+
+
+FReply STreeMap::InsertNewNodeAsChildOfHoveredNode( const FGeometry& MyGeometry )
+{
+	FReply Reply = FReply::Unhandled();
+
+	// Wait until we've finished animating a transition
+	if( !IsNavigationTransitionActive() )
+	{
+		if( MouseOverVisual != NULL )
+		{
+			auto ParentNode = MouseOverVisual->NodeData->AsShared();
+
+			// Allocate the new node but don't add it to the tree yet.  We want the user to give the node a name first!
+
+			// NOTE: Ownership of this node will be transferred to StartRenamingNode, where it will be referenced
+			// in the delegate callback for the editable text change commit handler.  If the user opts to not type anything, the
+			// node will be deleted instead of being added to the tree.
+			FTreeMapNodeDataRef NewNode( new FTreeMapNodeData() );
+			NewNode->Name = TEXT( "" );
+			NewNode->Size = 1.0f;	// Leaf nodes always get a size of 1.0 for now
+			NewNode->Parent = MouseOverVisual->NodeData;
+
+			// Start editing the new node before we insert it
+			const bool bIsNewNode = true;
+			StartRenamingNode( MyGeometry, NewNode, MouseOverVisual->Position, bIsNewNode );
+
+			Reply = FReply::Handled();
+		}
+	}
+
+	return Reply;
+}
+
+
+bool STreeMap::SupportsKeyboardFocus() const
+{
+	return true;
+}
+
+
+FReply STreeMap::OnKeyDown( const FGeometry& MyGeometry, const FKeyEvent& InKeyboardEvent )
+{
+	FReply Reply = FReply::Unhandled();
+
+	const FKey Key = InKeyboardEvent.GetKey();
+
+	if( AllowEditing.Get() )
+	{
+		if( Key == EKeys::Delete )
+		{
+			// If the cursor is over a node, delete it!
+			DeleteHoveredNode();
+			Reply = FReply::Handled();
+		}
+		else if( Key == EKeys::Insert )
+		{
+			// If the cursor is over a node, insert a new node as a child!
+			InsertNewNodeAsChildOfHoveredNode( MyGeometry );
+			Reply = FReply::Handled();
+		}
+		else if( Key == EKeys::Z && InKeyboardEvent.IsControlDown() )
+		{
+			// Undo
+			Undo();
+			Reply = FReply::Handled();
+		}
+		else if( Key == EKeys::Y && InKeyboardEvent.IsControlDown() )
+		{
+			// Redo
+			Redo();
+			Reply = FReply::Handled();
+		}
+	}
+
+	return Reply;
+}
+
+
+void STreeMap::RenamingNode_OnTextCommitted( const FText& NewText, ETextCommit::Type, TSharedRef<FTreeMapNodeData> NodeToRename )
+{
+	TSharedPtr<SWidget> RenamingNodeWidgetPinned( RenamingNodeWidget.Pin() );
+	if( RenamingNodeWidgetPinned.IsValid() )
+	{
+		// Kill the window after the text has been committed
+		TSharedRef<SWindow> ParentWindow = FSlateApplication::Get().FindWidgetWindow( RenamingNodeWidgetPinned.ToSharedRef() ).ToSharedRef();
+
+		RenamingNodeWidget.Reset();	// Avoid reentrancy
+		RenamingNodeData.Reset();
+		FSlateApplication::Get().RequestDestroyWindow( ParentWindow );
+
+
+		// Make sure new name is OK
+		if( NewText.ToString() != NodeToRename->Name &&
+			!NewText.IsEmpty() )
+		{
+			// Store undo state
+			TakeUndoSnapshot();
+
+			// Rename it!
+			NodeToRename->Name = NewText.ToString();
+
+			if( bIsNamingNewNode )
+			{
+				// Add the new node to the tree
+				TakeUndoSnapshot();
+
+				// Insert the new node
+				{
+					auto ParentNode = NodeToRename->Parent->AsShared();
+					ParentNode->Children.Add( NodeToRename );
+				}
+
+				const bool bShouldPlayTransition = true;
+				RebuildTreeMap( bShouldPlayTransition );
+			}
+			else
+			{
+				// NOTE: No refresh needed, as node labels are pulled directly from nodes
+			}
+		}
+	}
+}
+
+
+
+void STreeMap::StartRenamingNode( const FGeometry& MyGeometry, const FTreeMapNodeDataRef& NodeData, const FVector2D& RelativePosition, const bool bIsNewNode )
+{
+	TSharedRef<SBorder> RenamerWidget = SNew( SBorder );
+
+	TSharedRef<SEditableTextBox> EditableText = 
+		SNew( SEditableTextBox )
+			.Text( FText::FromString( NodeData->Name ) )
+			.SelectAllTextWhenFocused( true )
+			.RevertTextOnEscape( true )
+			.MinDesiredWidth( 100 )
+			.OnTextCommitted( this, &STreeMap::RenamingNode_OnTextCommitted, NodeData->AsShared() )
+		;
+	RenamerWidget->SetContent( EditableText );
+
+	RenamingNodeWidget = RenamerWidget;
+	RenamingNodeData = NodeData;
+	bIsNamingNewNode = bIsNewNode;
+
+	const bool bFocusImmediately = true;
+	const bool bShouldAutoSize = true;
+	FSlateApplication::Get().PushMenu( AsShared(), RenamerWidget, MyGeometry.LocalToAbsolute( RelativePosition ), FPopupTransitionEffect::None, bFocusImmediately, bShouldAutoSize );
+
+	// Focus the text box right after we spawn it so that the user can start typing
+	FSlateApplication::Get().SetKeyboardFocus( EditableText, EKeyboardFocusCause::SetDirectly );
+}
+
+
+
+void STreeMap::StopRenamingNode()
+{
+	TSharedPtr<SWidget> RenamingNodeWidgetPinned( RenamingNodeWidget.Pin() );
+	if( RenamingNodeWidgetPinned.IsValid() )
+	{
+		// Kill the window after the text has been committed
+		TSharedRef<SWindow> ParentWindow = FSlateApplication::Get().FindWidgetWindow( RenamingNodeWidgetPinned.ToSharedRef() ).ToSharedRef();
+		FSlateApplication::Get().RequestDestroyWindow( ParentWindow );
+	}
+}
+
+
+void STreeMap::ApplyVisualizationToNodes( const FTreeMapNodeDataRef& Node )
+{
+	const FLinearColor RootDefaultColor = FLinearColor( 0.125f, 0.125f, 0.125f );	// @todo treemap custom: Make configurable in the STreeMap settings?
+	int32 TreeDepth = 0;
+	ApplyVisualizationToNodesRecursively( Node, RootDefaultColor, TreeDepth );
+}
+
+
+void STreeMap::ApplyVisualizationToNodesRecursively( const FTreeMapNodeDataRef& Node, const FLinearColor& DefaultColor, const int32 TreeDepth )
+{
+	// Select default color saturation based on tree depth
+	FLinearColor MyDefaultColor = DefaultColor;
+	auto HSV = MyDefaultColor.LinearRGBToHSV();
+	const float SaturationReductionPerDepthLevel = 0.2f;	// @todo treemap custom: Make configurable in the tree map settings?  Calculate tree depth?
+	const float MinAllowedSaturation = 0.1f;
+	HSV.G = FMath::Max( MinAllowedSaturation, HSV.G - TreeDepth * SaturationReductionPerDepthLevel );
+	MyDefaultColor = HSV.HSVToLinearRGB();
+
+
+	// Size
+	{
+		if( SizeNodesByAttribute.IsValid() )
+		{
+			FTreeMapAttributeDataPtr* AttributeDataPtr = Node->Attributes.Find( SizeNodesByAttribute->Name );
+			if( AttributeDataPtr != NULL )
+			{
+				const FTreeMapAttributeData& AttributeData = *AttributeDataPtr->Get();
+
+				// Make sure the data value is in range
+				if( ensure( SizeNodesByAttribute->Values.Contains( AttributeData.Value ) ) )
+				{
+					// @todo treemap: This doesn't work well with container nodes.  Our range of sizes doesn't go big enough to compete with cumulative sizes!
+					Node->Size = SizeNodesByAttribute->Values[ AttributeData.Value ]->NodeSize;
+				}
+				else
+				{
+					// Invalid attribute data!
+					Node->Size = Node->IsLeafNode() ? SizeNodesByAttribute->DefaultValue->NodeSize : 0.0f;	// Leaf nodes get a default size, but container nodes use auto-sizing
+				}
+			}
+			else
+			{
+				// The node doesn't have this attribute on it.  Use default.
+				Node->Size = Node->IsLeafNode() ? SizeNodesByAttribute->DefaultValue->NodeSize : 0.0f;	// Leaf nodes get a default size, but container nodes use auto-sizing
+			}
+		}
+
+		else
+		{
+			// @todo treemap: Really we want this to "restore the original size set by the user", not make up new defaults.  Disabled for now.
+
+			// Default size
+			// Node->Size = Node->IsLeafNode() ? 1.0f : 0.0f;	// Leaf nodes get a default size, but container nodes use auto-sizing
+		}
+	}
+
+
+	// Color
+	{
+		if( ColorNodesByAttribute.IsValid() )
+		{
+			FTreeMapAttributeDataPtr* AttributeDataPtr = Node->Attributes.Find( ColorNodesByAttribute->Name );
+			if( AttributeDataPtr != NULL )
+			{
+				const FTreeMapAttributeData& AttributeData = *AttributeDataPtr->Get();
+
+				// Make sure the data value is in range
+				if( ensure( ColorNodesByAttribute->Values.Contains( AttributeData.Value ) ) )
+				{
+					Node->Color = ColorNodesByAttribute->Values[ AttributeData.Value ]->NodeColor;
+				}
+				else
+				{
+					// Invalid attribute data!
+					Node->Color = ColorNodesByAttribute->DefaultValue->NodeColor;
+				}
+			}
+			else
+			{
+				// The node doesn't have this attribute on it.  Use default.
+				Node->Color = ColorNodesByAttribute->DefaultValue->NodeColor;
+			}
+		}
+		else
+		{
+			// Default color
+			Node->Color = MyDefaultColor;
+		}
+	}
+
+
+
+	for( int32 ChildIndex = 0; ChildIndex < Node->Children.Num(); ++ChildIndex )
+	{
+		const auto& ChildNode = Node->Children[ ChildIndex ];
+
+		// Make up a distinct color for all of the root's top level nodes
+		FLinearColor ChildColor = MyDefaultColor;
+		if( TreeDepth == 0 )
+		{
+			// Choose a hue evenly spaced across the spectrum
+			float ColorHue = 360.0f * (float)( ChildIndex + 1 ) / (float)Node->Children.Num();
+			auto ChildColorHSV = FLinearColor::White;
+			ChildColorHSV.R = ColorHue;
+			ChildColorHSV.G = 1.0f;	// Full saturation!
+			ChildColor = ChildColorHSV.HSVToLinearRGB();
+		}
+
+		ApplyVisualizationToNodesRecursively( ChildNode.ToSharedRef(), ChildColor, TreeDepth + 1 );
+	}
+
+}
+
+
+
+void STreeMap::ShowOptionsMenuAt( const FVector2D& ScreenSpacePosition )
+{
+	struct Local
+	{
+		static void MakeEditNodeAttributeMenu( FMenuBuilder& MenuBuilder, FTreeMapNodeDataPtr EditingNode, TSharedPtr< FTreeMapAttribute > Attribute, STreeMap* Self )
+		{
+			MenuBuilder.AddMenuEntry(
+				LOCTEXT( "RemoveAttribute", "Not Set" ),
+				FText(),	// No tooltip (intentional)
+				FSlateIcon(),	// Icon
+				FUIAction(
+				FExecuteAction::CreateStatic( &Local::EditNodeAttribute_Execute, EditingNode, Attribute, TSharedPtr< FTreeMapAttributeValue >(), Self ),
+				FCanExecuteAction(),
+				FIsActionChecked::CreateStatic( &Local::EditNodeAttribute_IsChecked, EditingNode, Attribute, TSharedPtr< FTreeMapAttributeValue >(), Self ) ),
+				NAME_None,	// Extension point
+				EUserInterfaceActionType::ToggleButton );
+
+			MenuBuilder.AddMenuSeparator();
+
+			// @todo treemap: These probably should be sorted rather than hash order
+			for( auto HashIter( Attribute->Values.CreateConstIterator() ); HashIter; ++HashIter )
+			{
+				FTreeMapAttributeValuePtr Value = HashIter.Value();
+
+				MenuBuilder.AddMenuEntry(
+					FText::FromName( Value->Name ),
+					FText(),	// No tooltip (intentional)
+					FSlateIcon(),	// Icon
+					FUIAction(
+					FExecuteAction::CreateStatic( &Local::EditNodeAttribute_Execute, EditingNode, Attribute, Value, Self ),
+					FCanExecuteAction(),
+					FIsActionChecked::CreateStatic( &Local::EditNodeAttribute_IsChecked, EditingNode, Attribute, Value, Self ) ),
+					NAME_None,	// Extension point
+					EUserInterfaceActionType::ToggleButton );
+			}
+		}
+
+
+		static void EditNodeAttribute_Execute( FTreeMapNodeDataPtr EditingNode, TSharedPtr< FTreeMapAttribute > Attribute, TSharedPtr< FTreeMapAttributeValue > Value, STreeMap* Self )
+		{
+			bool bAnyChanges = false;
+			if( Value.IsValid() )
+			{
+				if( !EditingNode->Attributes.Contains( Attribute->Name ) )
+				{
+					EditingNode->Attributes.Add( Attribute->Name, MakeShareable( new FTreeMapAttributeData( Value->Name ) ) );
+					bAnyChanges = true;
+				}
+				else if( EditingNode->Attributes[ Attribute->Name ]->Value != Value->Name )
+				{
+					EditingNode->Attributes[ Attribute->Name ]->Value = Value->Name;
+					bAnyChanges = true;
+				}
+				else
+				{
+					// Nothing changed
+				}
+			}
+			else
+			{
+				// Clearing attribute
+				if( EditingNode->Attributes.Contains( Attribute->Name ) )
+				{
+					EditingNode->Attributes.Remove( Attribute->Name );
+					bAnyChanges = true;
+				}
+				else
+				{
+					// Nothing changed
+				}
+			}
+
+			// Has anything changed?
+			if( bAnyChanges )
+			{
+				const bool bShouldPlayTransition = true;
+				Self->RebuildTreeMap( bShouldPlayTransition );
+			}
+		}
+
+
+		static bool EditNodeAttribute_IsChecked( FTreeMapNodeDataPtr EditingNode, TSharedPtr< FTreeMapAttribute > Attribute, TSharedPtr< FTreeMapAttributeValue > Value, STreeMap* Self )
+		{
+			if( Value.IsValid() )
+			{
+				return EditingNode->Attributes.Contains( Attribute->Name ) && EditingNode->Attributes[ Attribute->Name ]->Value == Value->Name;
+			}
+			else
+			{
+				return !EditingNode->Attributes.Contains( Attribute->Name );
+			}
+		}
+
+
+		static void SizeByAttribute_Execute( TSharedPtr< FTreeMapAttribute > Attribute, STreeMap* Self )
+		{
+			// Has anything changed?
+			if( Self->SizeNodesByAttribute != Attribute )
+			{
+				// Apply the new visualization!
+				Self->SizeNodesByAttribute = Attribute;
+
+				const bool bShouldPlayTransition = true;
+				Self->RebuildTreeMap( bShouldPlayTransition );
+			}
+		}
+
+
+		static bool SizeByAttribute_IsChecked( TSharedPtr< FTreeMapAttribute > Attribute, STreeMap* Self )
+		{
+			return Self->SizeNodesByAttribute == Attribute;
+		}
+
+
+		static void MakeSizeByAttributeMenu( FMenuBuilder& MenuBuilder, STreeMap* Self )
+		{
+			MenuBuilder.AddMenuEntry(
+				LOCTEXT( "NoSizeByAttribute", "Off" ),
+				FText(),	// No tooltip (intentional)
+				FSlateIcon(),	// Icon
+				FUIAction(
+				FExecuteAction::CreateStatic( &Local::SizeByAttribute_Execute, TSharedPtr< FTreeMapAttribute >(), Self ),
+				FCanExecuteAction(),
+				FIsActionChecked::CreateStatic( &Local::SizeByAttribute_IsChecked, TSharedPtr< FTreeMapAttribute >(), Self ) ),
+				NAME_None,	// Extension point
+				EUserInterfaceActionType::ToggleButton );
+
+			if( Self->Customization.IsValid() )
+			{
+				MenuBuilder.AddMenuSeparator();
+
+				// @todo treemap: These probably should be sorted rather than hash order
+				for( auto HashIter( Self->Customization->GetAttributes().CreateConstIterator() ); HashIter; ++HashIter )
+				{
+					FTreeMapAttributePtr Attribute = HashIter.Value();
+
+					MenuBuilder.AddMenuEntry(
+						FText::FromName( Attribute->Name ),
+						FText(),	// No tooltip (intentional)
+						FSlateIcon(),	// Icon
+						FUIAction(
+						FExecuteAction::CreateStatic( &Local::SizeByAttribute_Execute, Attribute, Self ),
+						FCanExecuteAction(),
+						FIsActionChecked::CreateStatic( &Local::SizeByAttribute_IsChecked, Attribute, Self ) ),
+						NAME_None,	// Extension point
+						EUserInterfaceActionType::ToggleButton );
+				}
+			}
+		}
+
+
+		static void ColorByAttribute_Execute( TSharedPtr< FTreeMapAttribute > Attribute, STreeMap* Self )
+		{
+			// Has anything changed?
+			if( Self->ColorNodesByAttribute != Attribute )
+			{
+				// Apply the new visualization!
+				Self->ColorNodesByAttribute = Attribute;
+
+				const bool bShouldPlayTransition = true;
+				Self->RebuildTreeMap( bShouldPlayTransition );
+			}
+		}
+
+
+		static bool ColorByAttribute_IsChecked( TSharedPtr< FTreeMapAttribute > Attribute, STreeMap* Self )
+		{
+			return Self->ColorNodesByAttribute == Attribute;
+		}
+
+
+		static void MakeColorByAttributeMenu( FMenuBuilder& MenuBuilder, STreeMap* Self )
+		{
+			MenuBuilder.AddMenuEntry(
+				LOCTEXT( "NoColorByAttribute", "Off" ),
+				FText(),	// No tooltip (intentional)
+				FSlateIcon(),	// Icon
+				FUIAction(
+				FExecuteAction::CreateStatic( &Local::ColorByAttribute_Execute, TSharedPtr< FTreeMapAttribute >(), Self ),
+				FCanExecuteAction(),
+				FIsActionChecked::CreateStatic( &Local::ColorByAttribute_IsChecked, TSharedPtr< FTreeMapAttribute >(), Self ) ),
+				NAME_None,	// Extension point
+				EUserInterfaceActionType::ToggleButton );
+
+			if( Self->Customization.IsValid() )
+			{
+				MenuBuilder.AddMenuSeparator();
+
+				// @todo treemap: These probably should be sorted rather than hash order
+				for( auto HashIter( Self->Customization->GetAttributes().CreateConstIterator() ); HashIter; ++HashIter )
+				{
+					FTreeMapAttributePtr Attribute = HashIter.Value();
+
+					MenuBuilder.AddMenuEntry(
+						FText::FromName( Attribute->Name ),
+						FText(),	// No tooltip (intentional)
+						FSlateIcon(),	// Icon
+						FUIAction(
+						FExecuteAction::CreateStatic( &Local::ColorByAttribute_Execute, Attribute, Self ),
+						FCanExecuteAction(),
+						FIsActionChecked::CreateStatic( &Local::ColorByAttribute_IsChecked, Attribute, Self ) ),
+						NAME_None,	// Extension point
+						EUserInterfaceActionType::ToggleButton );
+				}
+			}
+		}
+	};
+
+	// Only present a context menu if the tree has been customized
+	if( Customization.IsValid() )
+	{
+		const bool bShouldCloseMenuAfterSelection = true;
+		FMenuBuilder OptionsMenuBuilder( bShouldCloseMenuAfterSelection, NULL );
+
+		if( AllowEditing.Get() && Customization.IsValid() && MouseOverVisual != NULL )
+		{
+			// Node editing options
+			OptionsMenuBuilder.BeginSection( NAME_None, LOCTEXT( "EditNodeAttributesSection", "Edit Node" ) );
+			{
+				const auto& EditingNode = MouseOverVisual->NodeData->AsShared();
+
+				for( auto HashIter( Customization->GetAttributes().CreateConstIterator() ); HashIter; ++HashIter )
+				{
+					FTreeMapAttributePtr Attribute = HashIter.Value();
+
+					OptionsMenuBuilder.AddSubMenu(
+						FText::FromName( Attribute->Name ),
+						LOCTEXT( "EditAttribute_ToolTip", "Edits this attribute on the node under the curosr." ),
+						FNewMenuDelegate::CreateStatic( &Local::MakeEditNodeAttributeMenu, FTreeMapNodeDataPtr( EditingNode ), Attribute, this ) );
+				}
+			}
+			OptionsMenuBuilder.EndSection();
+		}
+
+		OptionsMenuBuilder.BeginSection( NAME_None, LOCTEXT( "ChangeLayoutSection", "Layout" ) );
+		{
+			OptionsMenuBuilder.AddSubMenu( LOCTEXT( "SizeBy", "Size by" ), LOCTEXT( "SizeBy_ToolTip", "Sets which criteria to base the size of the nodes in the tree map on." ), FNewMenuDelegate::CreateStatic( &Local::MakeSizeByAttributeMenu, this ) );
+			OptionsMenuBuilder.AddSubMenu( LOCTEXT( "ColorBy", "Color by" ), LOCTEXT( "ColorBy_ToolTip", "Sets which criteria to base the color of the nodes in the tree map on." ), FNewMenuDelegate::CreateStatic( &Local::MakeColorByAttributeMenu, this ) );
+		}
+		OptionsMenuBuilder.EndSection();
+
+		TSharedRef< SWidget > WindowContent =
+			SNew( SBorder )
+			[
+				OptionsMenuBuilder.MakeWidget()
+			];
+
+		const bool bFocusImmediately = false;
+		const bool bShouldAutoSize = true;
+		FSlateApplication::Get().PushMenu( AsShared(), WindowContent, ScreenSpacePosition, FPopupTransitionEffect::ContextMenu, bFocusImmediately, bShouldAutoSize );
+	}
+}
+
+#undef LOCTEXT_NAMESPACE
\ No newline at end of file
diff --git a/Engine/Source/Developer/TreeMap/TreeMap.Build.cs b/Engine/Source/Developer/TreeMap/TreeMap.Build.cs
new file mode 100644
--- /dev/null
+++ b/Engine/Source/Developer/TreeMap/TreeMap.Build.cs
@@ -0,0 +1,13 @@
+// Copyright 1998-2015 Epic Games, Inc. All Rights Reserved.
+
+namespace UnrealBuildTool.Rules
+{
+	public class TreeMap : ModuleRules
+	{
+		public TreeMap(TargetInfo Target)
+		{
+	        PublicDependencyModuleNames.AddRange(new string[] { "Core", "CoreUObject", "SlateCore" });
+	        PrivateDependencyModuleNames.AddRange(new string[] { "Slate", "InputCore", "XmlParser" });
+		}
+	}
+}
diff --git a/Engine/Source/Developer/TreeMap/TreeMap.cpp b/Engine/Source/Developer/TreeMap/TreeMap.cpp
new file mode 100644
--- /dev/null
+++ b/Engine/Source/Developer/TreeMap/TreeMap.cpp
@@ -0,0 +1,714 @@
+// Copyright 1998-2013 Epic Games, Inc. All Rights Reserved.
+
+#include "TreeMapModule.h"
+#include "ITreeMap.h"
+#include "XmlParser.h"
+
+
+typedef TSharedPtr<class FTreeMapNode> FTreeMapNodePtr;
+typedef TSharedRef<class FTreeMapNode> FTreeMapNodeRef;
+
+
+/**
+ * Rectangle used for tree maps
+ */
+struct FTreeMapRect
+{
+	/** Position of the rectangle */
+	FVector2D Position;
+
+	/** Dimensions of the rectangle */
+	FVector2D Size;
+
+
+	/** Default constructor */
+	FTreeMapRect()
+		: Position( FVector2D::ZeroVector ),
+		  Size( FVector2D::ZeroVector )
+	{
+	}
+};
+
+
+
+/**
+ * Single node in a tree map, which may have any number of child nodes, each with their own children and so on
+ */
+class FTreeMapNode
+{
+
+public:
+
+	/** Pointer to the source data for this node */
+	FTreeMapNodeDataPtr Data;
+
+	/** List of child nodes */
+	TArray< FTreeMapNodePtr > Children;
+
+	/** For leaf nodes, the size of this node.  For non-leaf nodes, the size of all of my child nodes. */
+	float Size;
+
+	/** Node rectangle */
+	FTreeMapRect Rect;
+
+	/** Node rectangle, with padding applied */
+	FTreeMapRect PaddedRect;
+
+	/** Font to use for this node's title */
+	FSlateFontInfo NameFont;
+
+	/** Font to use for this node's second line title */
+	FSlateFontInfo Name2Font;
+
+	/** Font to use for this node's centered text */
+	FSlateFontInfo CenterTextFont;
+
+	/** True if the node is 'interactive'.  That is, we have enough room for a title area and padding for the node to be clicked on */
+	bool bIsInteractive;
+
+
+public:
+
+	/** Default constructor for FTreeMapNode */
+	FTreeMapNode( const FTreeMapNodeDataRef& InitNodeData );
+
+	/** @return Returns true if this is a leaf node */
+	bool IsLeafNode() const
+	{
+		return Children.Num() == 0;
+	}
+
+};
+
+
+
+/**
+ * Tree map object
+ */
+class FTreeMap : public ITreeMap
+{
+
+public:
+
+	/** Default constructor for FTreeMap */
+	FTreeMap( const FTreeMapOptions& Options, const FTreeMapNodeDataRef& RootNodeData );
+
+	/** ITreeMap interface */
+	virtual TArray<FTreeMapNodeVisualInfo> GetVisuals() override;
+
+
+private:
+
+	enum class ESplitDirection
+	{
+		Horizontal,
+		Vertical
+	};
+
+	/** Recursively sets up nodes in the tree */
+	void AddNodesRecursively( FTreeMapNodePtr& Node, const FTreeMapNodeDataRef& NodeData );
+
+	/** Recursively caches the size of each node.  Leaf nodes get the size of their source node, while non-leaf nodes are set to the total size of all of their children */
+	void CalculateNodeSizesRecursively( const FTreeMapNodeRef& Node, float& MaxNodeSize );
+	
+	/** Scales the specified node and all sub-nodes by the specified amount */
+	void ScaleNodesRecursively( const FTreeMapNodeRef& NodeToScale, const float ScaleFactor );
+
+	/** Sets up a node using the standard tree mapping algorithm */
+	void MakeStandardTreeNode( const FTreeMapOptions& Options, FTreeMap::ESplitDirection SplitDirection, const FTreeMapNodeRef& Node );
+
+	/** Sets up a node using a squarification method */
+	void MakeSquarifiedTreeNode( const FTreeMapOptions& Options, const FTreeMapNodeRef& Node );
+
+	/** Partitions nodes recursively */
+	void PartitionNodesRecursively( const FTreeMapOptions& Options, FTreeMap::ESplitDirection SplitDirection, const FTreeMapNodeRef& Node );
+
+	/** Pad all of the nodes in to make room for titles and border */
+	void PadNodesRecursively( const FTreeMapOptions& Options, const FTreeMapNodeRef& Node, const int32 TreeDepth );
+
+
+private:
+
+	/** Root node in the tree map */
+	FTreeMapNodePtr RootNode;
+};
+
+
+
+
+FTreeMapNode::FTreeMapNode( const FTreeMapNodeDataRef& InitNodeData )
+	: Data( InitNodeData ),
+	  Children(),
+	  Size( 0.0f ),
+	  Rect(),
+	  PaddedRect(),
+	  NameFont(),
+	  Name2Font(),
+	  CenterTextFont(),
+	  bIsInteractive( true )
+{
+}
+
+
+
+FTreeMap::FTreeMap( const FTreeMapOptions& Options, const FTreeMapNodeDataRef& RootNodeData )
+{
+	AddNodesRecursively( RootNode, RootNodeData );
+
+	// Cache the size of every node
+	float MaxNodeSize = 0.0f;
+	CalculateNodeSizesRecursively( RootNode.ToSharedRef(), MaxNodeSize );
+
+	// Also fix up the node sizes as we go.  We want the sizes to be proportional to the total display size
+	const float DisplaySize = Options.DisplayWidth * Options.DisplayHeight;
+	ScaleNodesRecursively( RootNode.ToSharedRef(), DisplaySize / MaxNodeSize );	// @todo treemap perf: Could use a scale factor w/ accessor instead of recursing here
+
+	// The root node has a fixed position and size
+	RootNode->Rect.Position = FVector2D::ZeroVector;
+	RootNode->Rect.Size = FVector2D( Options.DisplayWidth, Options.DisplayHeight );
+
+	// For regular tree types, we'll choose a "next split direction" that matches the aspect of the display area
+	const float DisplayAspect = Options.DisplayWidth / Options.DisplayHeight;
+	const bool bIsWiderThanTall = DisplayAspect >= 1.0f;
+	const auto SplitDirection = bIsWiderThanTall ? ESplitDirection::Horizontal : ESplitDirection::Vertical;
+	PartitionNodesRecursively( Options, SplitDirection, RootNode.ToSharedRef() );
+
+	// Now add space for titles and borders
+	const int32 TreeDepth = 0;
+	PadNodesRecursively( Options, RootNode.ToSharedRef(), TreeDepth );
+}
+
+
+void FTreeMap::AddNodesRecursively( FTreeMapNodePtr& OutNode, const FTreeMapNodeDataRef& NodeData )
+{
+	// Setup this node
+	OutNode = MakeShareable( new FTreeMapNode( NodeData ) );
+
+	// Add children
+	for( const auto& ChildNodeData : NodeData->Children )
+	{
+		FTreeMapNodePtr ChildNode;
+		AddNodesRecursively( ChildNode, ChildNodeData.ToSharedRef() );
+
+		OutNode->Children.Add( ChildNode );
+	}
+}
+
+
+void FTreeMap::CalculateNodeSizesRecursively( const FTreeMapNodeRef& Node, float& MaxNodeSize )
+{
+	// Is this a leaf node?  Leaf nodes will actually determine the size of non-leaf nodes.
+	if( Node->IsLeafNode() )
+	{
+		// NOTE: Size should really always be greater than zero here to get good results, but we don't want to assert.
+		Node->Size = Node->Data->Size;
+	}
+	else
+	{
+		// Update child node sizes
+		float TotalSizeOfChildren = 0.0f;
+		float MaxSizeOfChildren = 0.0f;
+		for( const auto& ChildNode : Node->Children )
+		{
+			CalculateNodeSizesRecursively( ChildNode.ToSharedRef(), MaxSizeOfChildren );
+			TotalSizeOfChildren += ChildNode->Size;
+		}
+
+		// Container node.  If a size was explicitly set, then we'll use that size.
+		if( Node->Data->Size > 0.0f )
+		{
+			Node->Size = Node->Data->Size;
+
+			// Scale the child nodes to fit into the forced container size
+			const float ScaleFactor = Node->Size / TotalSizeOfChildren;
+			{
+				for( const auto& ChildNode : Node->Children )
+				{
+					ScaleNodesRecursively( ChildNode.ToSharedRef(), ScaleFactor );
+				}
+			}
+		}
+		else
+		{
+			// Create a size for the node by summing it's child node sizes
+			Node->Size = TotalSizeOfChildren;
+		}
+
+		// Sort our children, largest to smallest
+		Node->Children.Sort( []( const FTreeMapNodePtr& A, const FTreeMapNodePtr& B ) { return A->Size > B->Size; } );
+	}
+
+	MaxNodeSize = FMath::Max( MaxNodeSize, Node->Size );
+}
+
+
+void FTreeMap::ScaleNodesRecursively( const FTreeMapNodeRef& NodeToScale, const float ScaleFactor )
+{
+	NodeToScale->Size *= ScaleFactor;
+
+	for( const auto& ChildNode : NodeToScale->Children )
+	{
+		ScaleNodesRecursively( ChildNode.ToSharedRef(), ScaleFactor );
+	}
+}
+
+
+void FTreeMap::MakeStandardTreeNode( const FTreeMapOptions& Options, FTreeMap::ESplitDirection SplitDirection, const FTreeMapNodeRef& Node )
+{
+	// Standard tree map algorithm.  We alternate between horizontal and vertical packing of children.  All children are packed 
+	// into a single row or column.  This makes it fairly easy to see the hierarchical structure of the tree, but yields really long rectangles!
+	FVector2D Offset( FVector2D::ZeroVector );
+	for( const auto& ChildNode : Node->Children )
+	{
+		ChildNode->Rect.Position = Node->Rect.Position + Offset;
+
+		const float ChildFractionOfParent = ChildNode->Size / Node->Size;
+		if( SplitDirection == ESplitDirection::Horizontal )
+		{
+			ChildNode->Rect.Size.X = Node->Rect.Size.X * ChildFractionOfParent;
+			ChildNode->Rect.Size.Y = Node->Rect.Size.Y;
+			Offset.X += ChildNode->Rect.Size.X;
+		}
+		else
+		{
+			ChildNode->Rect.Size.X = Node->Rect.Size.X;
+			ChildNode->Rect.Size.Y = Node->Rect.Size.Y * ChildFractionOfParent;
+			Offset.Y += ChildNode->Rect.Size.Y;
+		}
+	}
+}
+
+
+void FTreeMap::MakeSquarifiedTreeNode( const FTreeMapOptions& Options, const FTreeMapNodeRef& Node )
+{
+	// NOTE: This algorithm is explained in the paper titled, "Squarified Treemaps", by Mark Bruls, Kees Huizing, and Jarke J.van Wijk
+
+	// For squarification, we'll always choose the wider aspect direction at every split (ignoring incoming NextSplitDirection!)
+
+	struct Local
+	{
+		/** Figure out the highest aspect ratio of all of the blocks, given the length of the rectangle that we want to place these blocks into */
+		static float GetWorstAspectInRow( const TArray< FTreeMapNodePtr >& Nodes, const float SubRectShortestSide )
+		{
+			float MinSize = MAX_FLT;
+			float MaxSize = 0.0f;
+			float TotalSize = 0.0f;
+			for( const auto& Node : Nodes )
+			{
+				MinSize = FMath::Min( MinSize, Node->Size );
+				MaxSize = FMath::Max( MaxSize, Node->Size );
+				TotalSize += Node->Size;
+			}
+
+			float TotalSizeSquared = TotalSize * TotalSize;
+			float SubRectShortestSideSquared = SubRectShortestSide * SubRectShortestSide;
+
+			float WorstAspect = FMath::Max( ( SubRectShortestSideSquared * MaxSize ) / TotalSizeSquared, TotalSizeSquared / ( SubRectShortestSideSquared * MinSize ) );
+			return WorstAspect;
+		}
+
+
+		/** Incoming nodes should be sorted, largest to smallest */
+		static TArray<FTreeMapNodePtr> BuildRowFromNodes( TArray<FTreeMapNodePtr>& Nodes, const float SubRectShortestSide )
+		{
+			TArray<FTreeMapNodePtr> Row;
+					
+			// Add the first child node to our row
+			Row.Add( Nodes[ 0 ] );
+			Nodes.RemoveAt( 0 );
+
+			// If there are no more nodes to sort, then we're finished for now
+			if( Nodes.Num() > 0 )
+			{
+				auto NewRow = Row;
+				do
+				{
+					NewRow.Add( Nodes[ 0 ] );
+					if( GetWorstAspectInRow( Row, SubRectShortestSide ) > GetWorstAspectInRow( NewRow, SubRectShortestSide ) )
+					{
+						Row = NewRow;
+
+						// Claim the node from the original list
+						Nodes.RemoveAt( 0 );
+					}
+					else
+					{
+						break;
+					}
+				}
+				while( Nodes.Num() > 0 );
+			}
+
+			return Row;
+		}
+
+
+		/** Figures out which nodes will fit */
+		static void PlaceNodes( TArray<FTreeMapNodePtr>& Row, FTreeMapRect& SubRect )
+		{
+			float TotalRowSize = 0.0f;
+			for( const auto& Node : Row )
+			{
+				TotalRowSize += Node->Size;
+			}
+
+			const FVector2D SubRectMax = SubRect.Position + SubRect.Size;
+
+			FTreeMapRect PlacementRect = SubRect;
+			if( SubRect.Size.X < SubRect.Size.Y )
+			{
+				// Taller than wide
+				float RowHeight = TotalRowSize / SubRect.Size.X;
+				if( PlacementRect.Position.Y + RowHeight >= SubRectMax.Y )
+				{
+					RowHeight = SubRectMax.Y - PlacementRect.Position.Y;
+				}
+
+				for( int32 ColumnIndex = 0; ColumnIndex < Row.Num(); ++ColumnIndex )
+				{
+					auto& Node = Row[ ColumnIndex ];
+
+					float Width = Node->Size / RowHeight;
+					if( PlacementRect.Position.X + Width > SubRectMax.X || ( ColumnIndex + 1 ) == Row.Num() )
+					{
+						Width = SubRectMax.X - PlacementRect.Position.X;
+					}
+					Node->Rect.Position = PlacementRect.Position;
+					Node->Rect.Size.X = Width;
+					Node->Rect.Size.Y = RowHeight;
+
+					PlacementRect.Position.X += Width;
+				}
+
+				const float NewY = SubRect.Position.Y + RowHeight;
+				SubRect.Size.Y -= NewY - SubRect.Position.Y;
+				SubRect.Position.Y = NewY;
+			}
+			else
+			{
+				// Wider than tall
+				float RowWidth = TotalRowSize / SubRect.Size.Y;
+				if( PlacementRect.Position.X + RowWidth >= SubRectMax.X )
+				{
+					RowWidth = SubRectMax.X - PlacementRect.Position.X;
+				}
+
+				for( int32 ColumnIndex = 0; ColumnIndex < Row.Num(); ++ColumnIndex )
+				{
+					auto& Node = Row[ ColumnIndex ];
+
+					float Height = Node->Size / RowWidth;
+					if( PlacementRect.Position.Y + Height > SubRectMax.Y || ( ColumnIndex + 1 ) == Row.Num() )
+					{
+						Height = SubRectMax.Y - PlacementRect.Position.Y;
+					}
+					Node->Rect.Position = PlacementRect.Position;
+					Node->Rect.Size.X = RowWidth;
+					Node->Rect.Size.Y = Height;
+
+					PlacementRect.Position.Y += Height;
+				}
+
+				const float NewX = SubRect.Position.X + RowWidth;
+				SubRect.Size.X -= NewX - SubRect.Position.X;
+				SubRect.Position.X = NewX;
+			}
+		}
+	};
+
+	// Squarify it!
+	auto ChildrenCopy = Node->Children;
+	FTreeMapRect SubRect = Node->Rect;
+	do
+	{
+		const auto SubRectShortestSide = FMath::Min( SubRect.Size.X, SubRect.Size.Y );
+		auto Row = Local::BuildRowFromNodes( ChildrenCopy, SubRectShortestSide );
+		Local::PlaceNodes( Row, SubRect );
+	} 
+	while( ChildrenCopy.Num() > 0 );
+}
+
+
+void FTreeMap::PartitionNodesRecursively( const FTreeMapOptions& Options, FTreeMap::ESplitDirection SplitDirection, const FTreeMapNodeRef& Node )
+{
+	// Store off our padded copy of the rectangle.  We'll actually do the padding later on.
+	Node->PaddedRect = Node->Rect;
+
+	if( !Node->IsLeafNode() )
+	{
+		if( Options.TreeMapType == ETreeMapType::Standard )
+		{
+			MakeStandardTreeNode( Options, SplitDirection, Node );
+		}
+		else if( Options.TreeMapType == ETreeMapType::Squarified )
+		{
+			MakeSquarifiedTreeNode( Options, Node );
+		}
+
+		// The default algorithm just alternates between horizontal and vertical.  The squarification algorithm ignores this.
+		auto NextSplitDirection = ( SplitDirection == ESplitDirection::Horizontal ) ? ESplitDirection::Vertical : ESplitDirection::Horizontal;
+
+		// Process children
+		for( const auto& ChildNode : Node->Children )
+		{
+			PartitionNodesRecursively( Options, NextSplitDirection, ChildNode.ToSharedRef() );
+		}
+	}
+}
+
+
+void FTreeMap::PadNodesRecursively( const FTreeMapOptions& Options, const FTreeMapNodeRef& Node, const int32 TreeDepth )
+{
+	// Inset the container node to leave room for a border, if needed
+	// Don't inset the root node
+	const auto OriginalNodeRect = Node->Rect;
+
+	// Choose a height for this node's font
+	const uint16 MinAllowedFontSize = 8;		// @todo treemap custom: Don't hardcode and instead make this a customizable option?
+	Node->NameFont = Options.NameFont; 
+	Node->NameFont.Size = FMath::Max< uint16 >( MinAllowedFontSize, Options.NameFont.Size - ( TreeDepth * Options.FontSizeChangeBasedOnDepth ) );
+	Node->Name2Font = Options.Name2Font; 
+	Node->Name2Font.Size = FMath::Max< uint16 >( MinAllowedFontSize, Options.Name2Font.Size - ( TreeDepth * Options.FontSizeChangeBasedOnDepth ) );
+	Node->CenterTextFont = Options.CenterTextFont; 
+	Node->CenterTextFont.Size = FMath::Max< uint16 >( MinAllowedFontSize, Options.CenterTextFont.Size - ( TreeDepth * Options.FontSizeChangeBasedOnDepth ) );
+
+	if( Node != RootNode )
+	{
+		// Make sure we don't pad beyond our node's size
+		const float ContainerOuterPadding = TreeDepth == 1 ? Options.TopLevelContainerOuterPadding : Options.NestedContainerOuterPadding;
+		const FVector2D MaxPadding( Node->PaddedRect.Size * 0.5f );
+		const FVector2D Padding( FMath::Min( ContainerOuterPadding, MaxPadding.X ), FMath::Min( ContainerOuterPadding, MaxPadding.Y ) );
+
+		Node->PaddedRect.Position += Padding;
+		Node->PaddedRect.Size -= Padding * 2.0f;
+	}
+
+	{
+		// The 'child area' is the area within this node that we will fit all child nodes into
+		auto ChildAreaRect = Node->PaddedRect;
+
+		// Unless this is a top level node, make sure the node is big enough to bother reporting to our caller.  They may not want to visualize tiny nodes!
+		Node->bIsInteractive = ( TreeDepth <= 1 || ChildAreaRect.Size.X * ChildAreaRect.Size.Y >= Options.MinimumInteractiveNodeSize );
+		if( Node->bIsInteractive )
+		{
+			// Figure out how much space we need for the title text
+			const TSharedRef< FSlateFontMeasure >& FontMeasureService = FSlateApplication::Get().GetRenderer()->GetFontMeasureService();
+			const float MaxCharacterHeight = FontMeasureService->GetMaxCharacterHeight( Node->NameFont );		// @todo treemap perf: Cache this for various heights to reduce calls to FSlateFontMeasure
+			const float ContainerTitleAreaHeight = MaxCharacterHeight;
+
+			// Leave room for a title if we were asked to do that
+			{
+				const float Padding = FMath::Min( ChildAreaRect.Size.Y, ContainerTitleAreaHeight );
+				ChildAreaRect.Position.Y += Padding;
+				ChildAreaRect.Size.Y -= Padding;
+			}
+
+			// Apply inner padding before our child nodes, if needed
+			{
+				// Make sure we don't pad beyond our node's size
+				const FVector2D MaxPadding( ChildAreaRect.Size * 0.5f );
+				const FVector2D Padding( FMath::Min( Options.ContainerInnerPadding, MaxPadding.X ), FMath::Min( Options.ContainerInnerPadding, MaxPadding.Y ) );
+
+				ChildAreaRect.Position += Padding;
+				ChildAreaRect.Size -= Padding * 2.0f;
+			}
+		}
+
+		// Offset and scale all of the child node rects to fit into the child area.  This is where some squashing might happen, 
+		// and the sizes are no longer 1:1 with what they originally represented.  But for our purposes this is OK!  If you need
+		// the sizes to be perfectly accurate, then disable all padding options.
+		for( const auto& ChildNode : Node->Children )
+		{
+			ChildNode->PaddedRect.Position = ChildAreaRect.Position + ( ChildNode->PaddedRect.Position - OriginalNodeRect.Position ) / OriginalNodeRect.Size * ChildAreaRect.Size;
+			ChildNode->PaddedRect.Size = ChildNode->PaddedRect.Size / OriginalNodeRect.Size * ChildAreaRect.Size;
+		}
+	}
+
+	// Process children
+	for( const auto& ChildNode : Node->Children )
+	{
+		PadNodesRecursively( Options, ChildNode.ToSharedRef(), TreeDepth + 1 );
+	}
+}
+
+
+TSharedRef<ITreeMap> ITreeMap::CreateTreeMap( const FTreeMapOptions& Options, const FTreeMapNodeDataRef& RootNodeData )
+{
+	return MakeShareable( new FTreeMap( Options, RootNodeData ) );
+}
+
+
+
+FTreeMapNodeDataPtr ITreeMap::ParseOPMLToTreeMapData( const FString& OPMLFilePath, FString& OutErrorMessage )
+{
+	// Use the file name as the root node name
+	const FString RootNodeName = FPaths::GetBaseFilename( OPMLFilePath );
+
+	FXmlFile OPML;
+	bool bLoadResult = OPML.LoadFile( OPMLFilePath );
+
+	FTreeMapNodeDataPtr RootNodeData;
+	if( bLoadResult && OPML.IsValid() )
+	{
+		// Get the working Xml Node
+		const FXmlNode* XmlRoot = OPML.GetRootNode();
+		if( XmlRoot != nullptr )
+		{
+			const auto& RootName = XmlRoot->GetTag();
+			if( RootName.Equals( TEXT( "opml" ), ESearchCase::IgnoreCase ) )
+			{
+				for( const auto& XmlNode : XmlRoot->GetChildrenNodes() )
+				{
+					const auto& NodeName = XmlNode->GetTag();
+					if( NodeName.Equals( TEXT( "body" ), ESearchCase::IgnoreCase ) )
+					{
+						struct Local
+						{
+							static void RecursivelyCreateNodes( const FTreeMapNodeDataRef& NodeData, const FXmlNode& XmlNode )
+							{
+								for( const auto& ChildXmlNode : XmlNode.GetChildrenNodes() )
+								{
+									const auto& NodeName = ChildXmlNode->GetTag();
+									if( NodeName.Equals( TEXT( "outline" ), ESearchCase::IgnoreCase ) )
+									{
+										FTreeMapNodeDataRef ChildNodeData = MakeShareable( new FTreeMapNodeData() );
+										ChildNodeData->Parent = &NodeData.Get();
+
+										// All outline nodes MUST have a text attribute (required as part of OPML spec)
+										const auto& OutlineText = ChildXmlNode->GetAttribute( TEXT( "text" ) );
+										ChildNodeData->Name = OutlineText;
+										NodeData->Children.Add( ChildNodeData );
+										
+										// Recurse into children
+										RecursivelyCreateNodes( ChildNodeData, *ChildXmlNode );
+
+
+										// Setup attributes of this node
+										{
+											const float DefaultLeafNodeSize = 1.0f;			// Leaf nodes must always have a non-zero size!
+											const float DefaultContainerNodeSize = 0.0f;		// 0.0 for container nodes, means "compute my size using my children"
+											ChildNodeData->Size = ChildNodeData->IsLeafNode() ? DefaultLeafNodeSize : DefaultContainerNodeSize;
+
+											// Parse out any hash tags
+											int32 HashTagCharIndex;
+											while( ChildNodeData->Name.FindChar( '#', HashTagCharIndex ) )
+											{
+												// Parse hash tag string
+												int32 HashTagLength = 1;
+												while( ( ChildNodeData->Name.Len() > HashTagCharIndex + HashTagLength ) &&
+													   !FChar::IsWhitespace( ChildNodeData->Name[ HashTagCharIndex + HashTagLength ] ) &&
+													   ChildNodeData->Name[ HashTagCharIndex + HashTagLength ] != '#' )
+												{
+													++HashTagLength;
+												}
+
+
+												if( HashTagLength > 1 )
+												{
+													const auto HashTag = ChildNodeData->Name.Mid( HashTagCharIndex + 1, HashTagLength - 1 );
+
+													ChildNodeData->HashTags.Add( HashTag );
+
+													// Strip the hash tag ofg of the original string
+													ChildNodeData->Name = ChildNodeData->Name.Mid( 0, HashTagCharIndex );
+													if( HashTagCharIndex + HashTagLength < ChildNodeData->Name.Len() )
+													{
+														ChildNodeData->Name += ChildNodeData->Name.Mid( HashTagCharIndex + HashTagLength );
+													}
+												}
+											}
+
+											// Clean up any leftover whitespace in the node name, after stripping out hash tags
+											// @todo treemap perf: This is pretty inefficient (always makes copies even when nothing to do!)
+											ChildNodeData->Name.TrimTrailing();
+											ChildNodeData->Name.Trim();
+										}
+									}
+									else
+									{
+										// Node that we're not interested in
+									}
+								}
+							}
+						};
+
+						RootNodeData = MakeShareable( new FTreeMapNodeData() );
+						RootNodeData->Parent = NULL;
+						RootNodeData->Name = RootNodeName;
+						Local::RecursivelyCreateNodes( RootNodeData.ToSharedRef(), *XmlNode );
+					}
+					else
+					{
+						// Top level node that we're not interested in
+					}
+				}
+
+				if( !RootNodeData.IsValid() )
+				{
+					OutErrorMessage = TEXT( "Couldn't find a 'body' node in the XML document" );
+				}
+
+			}
+			else
+			{
+				OutErrorMessage = TEXT( "File does not appear to be an OPML-formatted XML document" );
+			}
+		}
+		else
+		{
+			OutErrorMessage = TEXT( "No root node found in XML document" );
+		}
+	}
+	else
+	{
+		// Couldn't load file
+		OutErrorMessage = OPML.GetLastError();
+	}
+
+	return RootNodeData;
+}
+
+
+TArray<FTreeMapNodeVisualInfo> FTreeMap::GetVisuals()
+{
+	TArray<FTreeMapNodeVisualInfo> Visuals;
+
+	struct Local
+	{
+		static void RecursivelyGatherVisuals( TArray<FTreeMapNodeVisualInfo>& VisualsList, const FTreeMapNodeRef& Node )
+		{
+			// Add a visual for the node that was passed in.  We'll recurse down into children afterwards.
+			FTreeMapNodeVisualInfo Visual;
+			Visual.NodeData = Node->Data.Get();
+			Visual.Position = Node->PaddedRect.Position;
+			Visual.Size = Node->PaddedRect.Size;
+			Visual.Color = Node->Data->Color;
+			Visual.NameFont = Node->NameFont;
+			Visual.Name2Font = Node->Name2Font;
+			Visual.CenterTextFont = Node->CenterTextFont;
+			Visual.bIsInteractive = Node->bIsInteractive;
+
+			// If the node is non-interactive, then ghost it
+			if( !Visual.bIsInteractive )
+			{
+				Visual.Color.A *= 0.25f;
+			}
+
+			VisualsList.Add( Visual );
+
+
+			// Process children
+			for( auto ChildNodeIndex = 0; ChildNodeIndex < Node->Children.Num(); ++ChildNodeIndex )
+			{
+				const auto& ChildNode = Node->Children[ ChildNodeIndex ];
+
+				// Make up a distinct color for all of the root's top level nodes
+				RecursivelyGatherVisuals( VisualsList, ChildNode.ToSharedRef() );
+			}
+		}
+	};
+
+	Local::RecursivelyGatherVisuals( Visuals, RootNode.ToSharedRef() );
+
+	return Visuals;
+}
diff --git a/Engine/Source/Developer/TreeMap/TreeMapModule.cpp b/Engine/Source/Developer/TreeMap/TreeMapModule.cpp
new file mode 100644
--- /dev/null
+++ b/Engine/Source/Developer/TreeMap/TreeMapModule.cpp
@@ -0,0 +1,30 @@
+// Copyright 1998-2015 Epic Games, Inc. All Rights Reserved.
+
+#include "TreeMapModule.h"
+#include "TreeMapStyle.h"
+
+/**
+ * Implements the TreeMap module.
+ */
+class FTreeMapModule
+	: public IModuleInterface
+{
+public:
+
+	// IModuleInterface interface
+	virtual void StartupModule() override
+	{
+		FTreeMapStyle::Initialize();
+	}
+
+	virtual void ShutdownModule() override
+	{
+		FTreeMapStyle::Shutdown();
+	}
+};
+
+
+IMPLEMENT_MODULE(FTreeMapModule, TreeMap);
+
+
+#undef LOCTEXT_NAMESPACE
diff --git a/Engine/Source/Developer/TreeMap/TreeMapModule.h b/Engine/Source/Developer/TreeMap/TreeMapModule.h
new file mode 100644
--- /dev/null
+++ b/Engine/Source/Developer/TreeMap/TreeMapModule.h
@@ -0,0 +1,8 @@
+// Copyright 1998-2015 Epic Games, Inc. All Rights Reserved.
+
+#pragma  once
+
+#include "Core.h"
+#include "CoreUObject.h"
+#include "SlateBasics.h"
+
diff --git a/Engine/Source/Developer/TreeMap/TreeMapStyle.cpp b/Engine/Source/Developer/TreeMap/TreeMapStyle.cpp
new file mode 100644
--- /dev/null
+++ b/Engine/Source/Developer/TreeMap/TreeMapStyle.cpp
@@ -0,0 +1,63 @@
+// Copyright 1998-2013 Epic Games, Inc. All Rights Reserved.
+
+#include "TreeMapModule.h"
+#include "TreeMapStyle.h"
+#include "SlateStyle.h"
+
+#define IMAGE_BRUSH( RelativePath, ... ) FSlateImageBrush( FPaths::EngineContentDir() / "Editor/Slate"/ RelativePath + TEXT(".png"), __VA_ARGS__ )
+#define BOX_BRUSH( RelativePath, ... ) FSlateBoxBrush( FPaths::EngineContentDir() / "Editor/Slate"/ RelativePath + TEXT(".png"), __VA_ARGS__ )
+#define BORDER_BRUSH( RelativePath, ... ) FSlateBorderBrush( FPaths::EngineContentDir() / "Editor/Slate"/ RelativePath + TEXT(".png"), __VA_ARGS__ )
+#define TTF_FONT( RelativePath, ... ) FSlateFontInfo( FPaths::EngineContentDir() / "Editor/Slate"/ RelativePath + TEXT(".ttf"), __VA_ARGS__ )
+#define OTF_FONT( RelativePath, ... ) FSlateFontInfo( FPaths::EngineContentDir() / "Editor/Slate"/ RelativePath + TEXT(".otf"), __VA_ARGS__ )
+
+TSharedPtr< FSlateStyleSet > FTreeMapStyle::StyleInstance = nullptr;
+
+void FTreeMapStyle::Initialize()
+{
+	if (!StyleInstance.IsValid())
+	{
+		StyleInstance = Create();
+		FSlateStyleRegistry::RegisterSlateStyle(*StyleInstance);
+	}
+}
+
+void FTreeMapStyle::Shutdown()
+{
+	FSlateStyleRegistry::UnRegisterSlateStyle(*StyleInstance);
+	ensure(StyleInstance.IsUnique());
+	StyleInstance.Reset();
+}
+
+FName FTreeMapStyle::GetStyleSetName()
+{
+	static FName StyleSetName(TEXT("TreeMapStyle"));
+	return StyleSetName;
+}
+
+TSharedRef< FSlateStyleSet > FTreeMapStyle::Create()
+{
+	TSharedRef<FSlateStyleSet> StyleRef = MakeShareable(new FSlateStyleSet(FTreeMapStyle::GetStyleSetName()));
+	StyleRef->SetContentRoot(FPaths::EngineContentDir() / TEXT("Editor/Slate"));
+	StyleRef->SetCoreContentRoot(FPaths::EngineContentDir() / TEXT("Slate"));
+
+	FSlateStyleSet& Style = StyleRef.Get();
+
+	// Tree Map
+	Style.Set( "TreeMap.Background", new BOX_BRUSH( "TreeMap/TreeMap_Background", FMargin( 5.f / 12.f ) ) );
+	Style.Set( "TreeMap.NodeBackground", new BOX_BRUSH( "TreeMap/TreeMap_NodeBackground", FMargin( 5.f / 12.f ) ) );
+	Style.Set( "TreeMap.HoveredNodeBackground", new BOX_BRUSH( "TreeMap/TreeMap_NodeHoveredBackground", FMargin(5.f/12.f) ) );
+	Style.Set( "TreeMap.BorderPadding", FVector2D(1,0) );
+
+	return StyleRef;
+}
+
+#undef IMAGE_BRUSH
+#undef BOX_BRUSH
+#undef BORDER_BRUSH
+#undef TTF_FONT
+#undef OTF_FONT
+
+const ISlateStyle& FTreeMapStyle::Get()
+{
+	return *StyleInstance;
+}
diff --git a/Engine/Source/Editor/ContentBrowser/ContentBrowser.Build.cs b/Engine/Source/Editor/ContentBrowser/ContentBrowser.Build.cs
--- a/Engine/Source/Editor/ContentBrowser/ContentBrowser.Build.cs
+++ b/Engine/Source/Editor/ContentBrowser/ContentBrowser.Build.cs
@@ -17,6 +17,7 @@ public class ContentBrowser : ModuleRules
 				"SourceControl",
 				"SourceControlWindows",
                 "ReferenceViewer",
+                "SizeMap",
 			}
 		);
 
@@ -49,7 +50,8 @@ public class ContentBrowser : ModuleRules
 				"EditorWidgets",
 				"GameProjectGeneration",
                 "MainFrame",
-                "ReferenceViewer"
+                "ReferenceViewer",
+                "SizeMap",
 			}
 		);
 		
diff --git a/Engine/Source/Editor/ContentBrowser/Private/AssetContextMenu.cpp b/Engine/Source/Editor/ContentBrowser/Private/AssetContextMenu.cpp
--- a/Engine/Source/Editor/ContentBrowser/Private/AssetContextMenu.cpp
+++ b/Engine/Source/Editor/ContentBrowser/Private/AssetContextMenu.cpp
@@ -13,6 +13,7 @@
 #include "Toolkits/ToolkitManager.h"
 #include "ConsolidateWindow.h"
 #include "ReferenceViewer.h"
+#include "ISizeMapModule.h"
 
 #include "ReferencedAssetsUtils.h"
 
@@ -447,6 +448,15 @@ bool FAssetContextMenu::AddReferenceMenuOptions(FMenuBuilder& MenuBuilder)
 				FExecuteAction::CreateSP( this, &FAssetContextMenu::ExecuteShowReferenceViewer )
 				)
 			);
+
+		MenuBuilder.AddMenuEntry(
+			LOCTEXT("SizeMap", "Size Map..."),
+			LOCTEXT("SizeMapTooltip", "Shows an interactive map of the approximate memory used by this asset and everything it references."),
+			FSlateIcon(),
+			FUIAction(
+				FExecuteAction::CreateSP( this, &FAssetContextMenu::ExecuteShowSizeMap )
+				)
+			);
 	}
 	MenuBuilder.EndSection();
 
@@ -1327,6 +1337,20 @@ void FAssetContextMenu::ExecuteShowReferenceViewer()
 	}
 }
 
+void FAssetContextMenu::ExecuteShowSizeMap()
+{
+	TArray<FName> PackageNames;
+	for ( auto AssetIt = SelectedAssets.CreateConstIterator(); AssetIt; ++AssetIt )
+	{
+		PackageNames.Add(AssetIt->PackageName);
+	}
+
+	if ( PackageNames.Num() > 0 )
+	{
+		ISizeMapModule::Get().InvokeSizeMapTab(PackageNames);
+	}
+}
+
 void FAssetContextMenu::ExecuteGoToCodeForAsset(UClass* SelectedClass)
 {
 	if (SelectedClass)
diff --git a/Engine/Source/Editor/ContentBrowser/Private/AssetContextMenu.h b/Engine/Source/Editor/ContentBrowser/Private/AssetContextMenu.h
--- a/Engine/Source/Editor/ContentBrowser/Private/AssetContextMenu.h
+++ b/Engine/Source/Editor/ContentBrowser/Private/AssetContextMenu.h
@@ -155,6 +155,9 @@ private:
 	/** Handler for ShowReferenceViewer */
 	void ExecuteShowReferenceViewer();
 
+	/** Handler for ShowSizeMap */
+	void ExecuteShowSizeMap();
+
 	/** Handler for GoToAssetCode */
 	void ExecuteGoToCodeForAsset(UClass* SelectedClass);
 
diff --git a/Engine/Source/Editor/ContentBrowser/Private/PathContextMenu.cpp b/Engine/Source/Editor/ContentBrowser/Private/PathContextMenu.cpp
--- a/Engine/Source/Editor/ContentBrowser/Private/PathContextMenu.cpp
+++ b/Engine/Source/Editor/ContentBrowser/Private/PathContextMenu.cpp
@@ -7,6 +7,7 @@
 #include "SourceControlWindows.h"
 #include "ContentBrowserModule.h"
 #include "ReferenceViewer.h"
+#include "ISizeMapModule.h"
 #include "AssetToolsModule.h"
 #include "Editor/UnrealEd/Public/PackageTools.h"
 #include "SColorPicker.h"
@@ -227,6 +228,14 @@ void FPathContextMenu::MakePathViewContextMenu(FMenuBuilder& MenuBuilder)
 					FUIAction( FExecuteAction::CreateSP( this, &FPathContextMenu::ExecuteReferenceViewer ) )
 					);
     
+				// Size Map
+				MenuBuilder.AddMenuEntry(
+					LOCTEXT("SizeMap", "Size Map..."),
+					LOCTEXT("SizeMapOnFolderTooltip", "Shows an interactive map of the approximate memory used by the assets in this folder and everything they reference."),
+					FSlateIcon(),
+					FUIAction( FExecuteAction::CreateSP( this, &FPathContextMenu::ExecuteSizeMap ) )
+					);
+
 				// Fix Up Redirectors in Folder
 				MenuBuilder.AddMenuEntry(
 					LOCTEXT("FixUpRedirectorsInFolder", "Fix Up Redirectors in Folder"),
@@ -655,6 +664,23 @@ void FPathContextMenu::ExecuteReferenceViewer()
 	}
 }
 
+void FPathContextMenu::ExecuteSizeMap()
+{
+	TArray<FString> PackageNamesAsStrings;
+	GetPackageNamesInSelectedPaths(PackageNamesAsStrings);
+
+	TArray<FName> PackageNames;
+	for ( auto PackageNameIt = PackageNamesAsStrings.CreateConstIterator(); PackageNameIt; ++PackageNameIt )
+	{
+		PackageNames.Add(**PackageNameIt);
+	}
+
+	if ( PackageNames.Num() > 0 )
+	{
+		ISizeMapModule::Get().InvokeSizeMapTab(PackageNames);
+	}
+}
+
 void FPathContextMenu::ExecuteFixUpRedirectorsInFolder()
 {
 	FAssetRegistryModule& AssetRegistryModule = FModuleManager::Get().LoadModuleChecked<FAssetRegistryModule>(TEXT("AssetRegistry"));
diff --git a/Engine/Source/Editor/ContentBrowser/Private/PathContextMenu.h b/Engine/Source/Editor/ContentBrowser/Private/PathContextMenu.h
--- a/Engine/Source/Editor/ContentBrowser/Private/PathContextMenu.h
+++ b/Engine/Source/Editor/ContentBrowser/Private/PathContextMenu.h
@@ -81,6 +81,9 @@ public:
 	/** Handler for when "ReferenceViewer" is selected */
 	void ExecuteReferenceViewer();
 
+	/** Handler for when "SizeMap" is selected */
+	void ExecuteSizeMap();
+
 	/** Handler for when "Fix up Redirectors in Folder" is selected */
 	void ExecuteFixUpRedirectorsInFolder();
 
diff --git a/Engine/Source/Editor/LevelEditor/LevelEditor.Build.cs b/Engine/Source/Editor/LevelEditor/LevelEditor.Build.cs
--- a/Engine/Source/Editor/LevelEditor/LevelEditor.Build.cs
+++ b/Engine/Source/Editor/LevelEditor/LevelEditor.Build.cs
@@ -13,6 +13,7 @@ public class LevelEditor : ModuleRules
 				"MainFrame",
                 "PlacementMode",
 				"ReferenceViewer",
+				"SizeMap",
                 "IntroTutorials",
                 "AppFramework"
 			}
@@ -82,6 +83,7 @@ public class LevelEditor : ModuleRules
                 "PlacementMode",
 				"UserFeedback",
 				"ReferenceViewer",
+				"SizeMap",
                 "IntroTutorials"
 			}
 		);
diff --git a/Engine/Source/Editor/LevelEditor/Private/LevelEditor.cpp b/Engine/Source/Editor/LevelEditor/Private/LevelEditor.cpp
--- a/Engine/Source/Editor/LevelEditor/Private/LevelEditor.cpp
+++ b/Engine/Source/Editor/LevelEditor/Private/LevelEditor.cpp
@@ -618,6 +618,12 @@ void FLevelEditorModule::BindGlobalLevelEditorCommands()
 		FCanExecuteAction::CreateStatic( &FLevelEditorActionCallbacks::CanViewReferences )
 		);
 
+	ActionList.MapAction( 
+		FGlobalEditorCommonCommands::Get().ViewSizeMap, 
+		FExecuteAction::CreateStatic( &FLevelEditorActionCallbacks::ViewSizeMap_Execute ),
+		FCanExecuteAction::CreateStatic( &FLevelEditorActionCallbacks::CanViewSizeMap )
+		);
+
 	const FVector* NullVector = nullptr;
 	ActionList.MapAction(
 		Commands.GoHere,
diff --git a/Engine/Source/Editor/LevelEditor/Private/LevelEditorActions.cpp b/Engine/Source/Editor/LevelEditor/Private/LevelEditorActions.cpp
--- a/Engine/Source/Editor/LevelEditor/Private/LevelEditorActions.cpp
+++ b/Engine/Source/Editor/LevelEditor/Private/LevelEditorActions.cpp
@@ -41,6 +41,7 @@
 #include "AnalyticsEventAttribute.h"
 #include "IAnalyticsProvider.h"
 #include "ReferenceViewer.h"
+#include "ISizeMapModule.h"
 #include "Developer/MeshUtilities/Public/MeshUtilities.h"
 #include "EditorClassUtils.h"
 #include "ComponentEditorUtils.h"
@@ -1057,6 +1058,38 @@ bool FLevelEditorActionCallbacks::CanViewReferences()
 	return ReferencedAssets.Num() > 0;
 }
 
+void FLevelEditorActionCallbacks::ViewSizeMap_Execute()
+{
+	if( GEditor->GetSelectedActorCount() > 0 )
+	{
+		TArray< UObject* > ReferencedAssets;
+		GEditor->GetReferencedAssetsForEditorSelection( ReferencedAssets );
+
+		if (ReferencedAssets.Num() > 0)
+		{
+			TArray< FName > ViewableObjects;
+			for( auto ObjectIter = ReferencedAssets.CreateConstIterator(); ObjectIter; ++ObjectIter )
+			{
+				// Don't allow user to perform certain actions on objects that aren't actually assets (e.g. Level Script blueprint objects)
+				const auto EditingObject = *ObjectIter;
+				if( EditingObject != NULL && EditingObject->IsAsset() )
+				{
+					ViewableObjects.Add( EditingObject->GetOuter()->GetFName());
+				}
+			}
+
+			ISizeMapModule::Get().InvokeSizeMapTab(ViewableObjects);
+		}
+	}
+}
+
+bool FLevelEditorActionCallbacks::CanViewSizeMap()
+{
+	TArray< UObject* > ReferencedAssets;
+	GEditor->GetReferencedAssetsForEditorSelection(ReferencedAssets);
+	return ReferencedAssets.Num() > 0;
+}
+
 void FLevelEditorActionCallbacks::EditAsset_Clicked( const EToolkitMode::Type ToolkitMode, TWeakPtr< SLevelEditor > LevelEditor, bool bConfirmMultiple )
 {
 	if( GEditor->GetSelectedActorCount() > 0 )
diff --git a/Engine/Source/Editor/LevelEditor/Private/LevelEditorContextMenu.cpp b/Engine/Source/Editor/LevelEditor/Private/LevelEditorContextMenu.cpp
--- a/Engine/Source/Editor/LevelEditor/Private/LevelEditorContextMenu.cpp
+++ b/Engine/Source/Editor/LevelEditor/Private/LevelEditorContextMenu.cpp
@@ -259,6 +259,7 @@ void FLevelEditorContextMenu::FillMenu( FMenuBuilder& MenuBuilder, TWeakPtr<SLev
 				}
 
 				MenuBuilder.AddMenuEntry(FGlobalEditorCommonCommands::Get().ViewReferences);
+				MenuBuilder.AddMenuEntry(FGlobalEditorCommonCommands::Get().ViewSizeMap);
 
 			}
 			MenuBuilder.EndSection();
diff --git a/Engine/Source/Editor/LevelEditor/Public/LevelEditorActions.h b/Engine/Source/Editor/LevelEditor/Public/LevelEditorActions.h
--- a/Engine/Source/Editor/LevelEditor/Public/LevelEditorActions.h
+++ b/Engine/Source/Editor/LevelEditor/Public/LevelEditorActions.h
@@ -790,6 +790,12 @@ public:
 	/** If true ViewReferences_Execute can be called */
 	static bool CanViewReferences();
 
+	/** Called when the ViewSizeMap command is executed */
+	static void ViewSizeMap_Execute();
+
+	/** If true ViewSizeMap_Execute can be called */
+	static bool CanViewSizeMap();
+
 	/** Called to when "Edit Asset" is clicked */
 	static void EditAsset_Clicked( const EToolkitMode::Type ToolkitMode, TWeakPtr< class SLevelEditor > LevelEditor, bool bAskMultiple );
 
diff --git a/Engine/Source/Editor/ReferenceViewer/Private/ReferenceViewer.cpp b/Engine/Source/Editor/ReferenceViewer/Private/ReferenceViewer.cpp
--- a/Engine/Source/Editor/ReferenceViewer/Private/ReferenceViewer.cpp
+++ b/Engine/Source/Editor/ReferenceViewer/Private/ReferenceViewer.cpp
@@ -59,13 +59,6 @@ public:
 		ReferenceViewer->SetGraphRootPackageNames(GraphRootPackageNames);
 	}
 
-	virtual TSharedRef<SWidget> CreateReferenceViewer(const TArray<FName>& GraphRootPackageNames) override
-	{
-		TSharedRef<SReferenceViewer> ReferenceViewer = SNew(SReferenceViewer);
-		ReferenceViewer->SetGraphRootPackageNames(GraphRootPackageNames);
-		return ReferenceViewer;
-	}
-
 private:
 	TSharedRef<SDockTab> SpawnReferenceViewerTab( const FSpawnTabArgs& SpawnTabArgs )
 	{
diff --git a/Engine/Source/Editor/ReferenceViewer/Private/ReferenceViewerActions.cpp b/Engine/Source/Editor/ReferenceViewer/Private/ReferenceViewerActions.cpp
--- a/Engine/Source/Editor/ReferenceViewer/Private/ReferenceViewerActions.cpp
+++ b/Engine/Source/Editor/ReferenceViewer/Private/ReferenceViewerActions.cpp
@@ -14,6 +14,7 @@ void FReferenceViewerActions::RegisterCommands()
 	UI_COMMAND(ReCenterGraph, "Re-Center Graph", "Re-centers the graph on this node, showing all referencers and dependencies for this asset instead", EUserInterfaceActionType::Button, FInputChord());
 	UI_COMMAND(ListReferencedObjects, "List Referenced Objects", "Shows a list of objects that the selected asset references.", EUserInterfaceActionType::Button, FInputChord());
 	UI_COMMAND(ListObjectsThatReference, "List Objects That Reference", "Lists objects that reference the selected asset.", EUserInterfaceActionType::Button, FInputChord());
+	UI_COMMAND(ShowSizeMap, "Size Map...", "Displays an interactive map showing the approximate size of this asset and everything it references", EUserInterfaceActionType::Button, FInputChord());
 	UI_COMMAND(ShowReferenceTree, "Show Reference Tree...", "Shows a reference tree for the selected asset.", EUserInterfaceActionType::Button, FInputChord());
 
 	UI_COMMAND(MakeLocalCollectionWithReferencedAssets, "Local", "Local. This collection is only visible to you and is not in source control.", EUserInterfaceActionType::Button, FInputChord());
diff --git a/Engine/Source/Editor/ReferenceViewer/Private/ReferenceViewerActions.h b/Engine/Source/Editor/ReferenceViewer/Private/ReferenceViewerActions.h
--- a/Engine/Source/Editor/ReferenceViewer/Private/ReferenceViewerActions.h
+++ b/Engine/Source/Editor/ReferenceViewer/Private/ReferenceViewerActions.h
@@ -32,6 +32,9 @@ public:
 	// Lists objects that reference the selected asset.
 	TSharedPtr<FUICommandInfo> ListObjectsThatReference;
 
+	// Shows a size map for the selected asset.
+	TSharedPtr<FUICommandInfo> ShowSizeMap;
+
 	// Shows a reference tree for the selected asset.
 	TSharedPtr<FUICommandInfo> ShowReferenceTree;
 
diff --git a/Engine/Source/Editor/ReferenceViewer/Private/ReferenceViewerSchema.cpp b/Engine/Source/Editor/ReferenceViewer/Private/ReferenceViewerSchema.cpp
--- a/Engine/Source/Editor/ReferenceViewer/Private/ReferenceViewerSchema.cpp
+++ b/Engine/Source/Editor/ReferenceViewer/Private/ReferenceViewerSchema.cpp
@@ -44,6 +44,7 @@ void UReferenceViewerSchema::GetContextMenuActions(const UEdGraph* CurrentGraph,
 							FNewMenuDelegate::CreateUObject(this, &UReferenceViewerSchema::GetMakeCollectionWithReferencedAssetsSubMenu )
 							);
 
+	MenuBuilder->AddMenuEntry(FReferenceViewerActions::Get().ShowSizeMap);
 	MenuBuilder->AddMenuEntry(FReferenceViewerActions::Get().ShowReferenceTree);
 }
 
diff --git a/Engine/Source/Editor/ReferenceViewer/Private/SReferenceViewer.cpp b/Engine/Source/Editor/ReferenceViewer/Private/SReferenceViewer.cpp
--- a/Engine/Source/Editor/ReferenceViewer/Private/SReferenceViewer.cpp
+++ b/Engine/Source/Editor/ReferenceViewer/Private/SReferenceViewer.cpp
@@ -8,6 +8,7 @@
 #include "ReferenceViewerActions.h"
 #include "EditorWidgets.h"
 #include "GlobalEditorCommonCommands.h"
+#include "ISizeMapModule.h"
 
 #include "Editor/UnrealEd/Public/ObjectTools.h"
 #include "Engine/Selection.h"
@@ -507,6 +508,11 @@ void SReferenceViewer::RegisterActions()
 		FExecuteAction::CreateSP(this, &SReferenceViewer::MakeCollectionWithReferencedAssets, ECollectionShareType::CST_Shared),
 		FCanExecuteAction::CreateSP(this, &SReferenceViewer::HasExactlyOneNodeSelected));
 	
+	ReferenceViewerActions->MapAction(
+		FReferenceViewerActions::Get().ShowSizeMap,
+		FExecuteAction::CreateSP(this, &SReferenceViewer::ShowSizeMap),
+		FCanExecuteAction::CreateSP(this, &SReferenceViewer::HasAtLeastOneNodeSelected));
+
 	ReferenceViewerActions->MapAction(
 		FReferenceViewerActions::Get().ShowReferenceTree,
 		FExecuteAction::CreateSP(this, &SReferenceViewer::ShowReferenceTree),
@@ -584,6 +590,27 @@ void SReferenceViewer::MakeCollectionWithReferencedAssets(ECollectionShareType::
 	}
 }
 
+void SReferenceViewer::ShowSizeMap()
+{
+	TArray<FName> SelectedAssetPackageNames;
+
+	TSet<UObject*> SelectedNodes = GraphEditorPtr->GetSelectedNodes();
+	if ( ensure(SelectedNodes.Num()) == 1 )
+	{
+		UEdGraphNode_Reference* ReferenceNode = Cast<UEdGraphNode_Reference>(SelectedNodes.Array()[0]);
+		if ( ReferenceNode )
+		{
+			const FAssetData& AssetData = ReferenceNode->GetAssetData();
+			SelectedAssetPackageNames.AddUnique( AssetData.PackageName );
+		}
+	}
+
+	if( SelectedAssetPackageNames.Num() > 0 )
+	{
+		ISizeMapModule::Get().InvokeSizeMapTab( SelectedAssetPackageNames );
+	}
+}
+
 void SReferenceViewer::ShowReferenceTree()
 {
 	UObject* SelectedObject = GetObjectFromSingleSelectedNode();
@@ -670,6 +697,16 @@ bool SReferenceViewer::HasExactlyOneNodeSelected() const
 	return false;
 }
 
+bool SReferenceViewer::HasAtLeastOneNodeSelected() const
+{
+	if ( GraphEditorPtr.IsValid() )
+	{
+		return GraphEditorPtr->GetSelectedNodes().Num() > 0;
+	}
+	
+	return false;
+}
+
 void SReferenceViewer::OnInitialAssetRegistrySearchComplete()
 {
 	if ( GraphObj )
diff --git a/Engine/Source/Editor/ReferenceViewer/Private/SReferenceViewer.h b/Engine/Source/Editor/ReferenceViewer/Private/SReferenceViewer.h
--- a/Engine/Source/Editor/ReferenceViewer/Private/SReferenceViewer.h
+++ b/Engine/Source/Editor/ReferenceViewer/Private/SReferenceViewer.h
@@ -78,12 +78,14 @@ private:
 	void ListReferencedObjects();
 	void ListObjectsThatReference();
 	void MakeCollectionWithReferencedAssets(ECollectionShareType::Type ShareType);
+	void ShowSizeMap();
 	void ShowReferenceTree();
 
 	void ReCenterGraphOnNodes(const TSet<UObject*>& Nodes);
 
 	UObject* GetObjectFromSingleSelectedNode() const;
 	bool HasExactlyOneNodeSelected() const;
+	bool HasAtLeastOneNodeSelected() const;
 
 	void OnInitialAssetRegistrySearchComplete();
 private:
diff --git a/Engine/Source/Editor/ReferenceViewer/Public/ReferenceViewer.h b/Engine/Source/Editor/ReferenceViewer/Public/ReferenceViewer.h
--- a/Engine/Source/Editor/ReferenceViewer/Public/ReferenceViewer.h
+++ b/Engine/Source/Editor/ReferenceViewer/Public/ReferenceViewer.h
@@ -36,8 +36,5 @@ public:
 
 	/** Invokes a major tab with a reference viewer within it */
 	virtual void InvokeReferenceViewerTab(const TArray<FName>& GraphRootPackageNames) = 0;
-
-	/** Creates the ReferenceViewer widget */
-	virtual TSharedRef<class SWidget> CreateReferenceViewer(const TArray<FName>& GraphRootPackageNames) = 0;
 };
 
diff --git a/Engine/Source/Editor/ReferenceViewer/ReferenceViewer.Build.cs b/Engine/Source/Editor/ReferenceViewer/ReferenceViewer.Build.cs
--- a/Engine/Source/Editor/ReferenceViewer/ReferenceViewer.Build.cs
+++ b/Engine/Source/Editor/ReferenceViewer/ReferenceViewer.Build.cs
@@ -26,13 +26,15 @@ namespace UnrealBuildTool.Rules
 				    "AssetRegistry",
 					"EditorWidgets",
 					"CollectionManager",
+					"SizeMap",
 			    }
             );
 
             DynamicallyLoadedModuleNames.AddRange(
                 new string[] {
 				    "AssetRegistry",
-					"EditorWidgets"
+					"EditorWidgets",
+					"SizeMap",
 			    }
             );
 		}
diff --git a/Engine/Source/Editor/SizeMap/Public/ISizeMapModule.h b/Engine/Source/Editor/SizeMap/Public/ISizeMapModule.h
new file mode 100644
--- /dev/null
+++ b/Engine/Source/Editor/SizeMap/Public/ISizeMapModule.h
@@ -0,0 +1,41 @@
+// Copyright 1998-2015 Epic Games, Inc. All Rights Reserved.
+
+#pragma once
+
+#include "ModuleManager.h"
+
+
+/**
+ * The public interface to this module
+ */
+class ISizeMapModule : public IModuleInterface
+{
+
+public:
+
+	/**
+	 * Singleton-like access to this module's interface.  This is just for convenience!
+	 * Beware of calling this during the shutdown phase, though.  Your module might have been unloaded already.
+	 *
+	 * @return Returns singleton instance, loading the module on demand if needed
+	 */
+	static inline ISizeMapModule& Get()
+	{
+		return FModuleManager::LoadModuleChecked< ISizeMapModule >( "SizeMap" );
+	}
+
+	/**
+	 * Checks to see if this module is loaded and ready.  It is only valid to call Get() if IsAvailable() returns true.
+	 *
+	 * @return True if the module is loaded and ready to use
+	 */
+	static inline bool IsAvailable()
+	{
+		return FModuleManager::Get().IsModuleLoaded( "SizeMap" );
+	}
+
+	/** Invokes a major tab with a SizeMap within it */
+	virtual void InvokeSizeMapTab(const TArray<FName>& AssetPackageNames) = 0;
+
+};
+
diff --git a/Engine/Source/Editor/SizeMap/SSizeMap.cpp b/Engine/Source/Editor/SizeMap/SSizeMap.cpp
new file mode 100644
--- /dev/null
+++ b/Engine/Source/Editor/SizeMap/SSizeMap.cpp
@@ -0,0 +1,561 @@
+// Copyright 1998-2015 Epic Games, Inc. All Rights Reserved.
+
+#include "SizeMapModule.h"
+#include "STreeMap.h"
+#include "SSizeMap.h"
+#include "AssetRegistryModule.h"
+#include "AssetThumbnail.h"
+#include "ClassIconFinder.h"
+#include "UnitConversion.h"
+
+#define LOCTEXT_NAMESPACE "SizeMap"
+
+
+SSizeMap::SSizeMap()
+	: TreeMapWidget( nullptr ),
+	  RootAssetPackageNames(),
+	  RootTreeMapNode( new FTreeMapNodeData() ),
+
+	  // @todo sizemap: Hard-coded thumbnail pool size.  Not a big deal, but ideally move the constants elsewhere
+  	  AssetThumbnailPool( new FAssetThumbnailPool(1024) )
+{
+}
+
+
+SSizeMap::~SSizeMap()
+{
+	if( AssetThumbnailPool.IsValid() )
+	{
+		AssetThumbnailPool->ReleaseResources();
+		AssetThumbnailPool.Reset();
+	}
+}
+
+
+void SSizeMap::Construct( const FArguments& InArgs )
+{
+	ChildSlot
+	[
+		SAssignNew( TreeMapWidget, STreeMap, RootTreeMapNode.ToSharedRef(), nullptr )
+			.OnTreeMapNodeDoubleClicked( this, &SSizeMap::OnTreeMapNodeDoubleClicked )
+	];
+}
+
+
+void SSizeMap::SetRootAssetPackageNames( const TArray<FName>& NewRootAssetPackageNames )
+{
+	RootAssetPackageNames = NewRootAssetPackageNames;
+
+	FAssetRegistryModule& AssetRegistryModule = FModuleManager::LoadModuleChecked<FAssetRegistryModule>(TEXT("AssetRegistry"));
+	if ( AssetRegistryModule.Get().IsLoadingAssets() )
+	{
+		// We are still discovering assets, listen for the completion delegate before building the graph
+		if (!AssetRegistryModule.Get().OnFilesLoaded().IsBoundToObject(this))
+		{		
+			AssetRegistryModule.Get().OnFilesLoaded().AddSP( this, &SSizeMap::OnInitialAssetRegistrySearchComplete );
+		}
+	}
+	else
+	{
+		// All assets are already discovered, build the graph now.
+		RefreshMap();
+	}
+}
+
+
+namespace SizeMapInternals
+{
+	/** Serialization archive that discovers assets referenced by a specific Unreal object */
+	class FAssetReferenceFinder : public FArchiveUObject
+	{
+	public:
+		FAssetReferenceFinder( UObject* Object )
+			: CurrentObject( Object )
+		{
+			ArIsObjectReferenceCollector = true;
+			ArIgnoreOuterRef = true;
+
+			CurrentObject->Serialize( *this );
+		}
+
+		FArchive& operator<<( UObject*& Object )
+		{
+			// Only look at objects which are valid
+			const bool bValidObject =
+				Object &&	// Object should not be NULL
+				!Object->HasAnyFlags( RF_Transient | RF_PendingKill ) && // Should not be transient or pending kill
+				Object->HasAllFlags( RF_Public ) &&	// All assets should be public
+				Object->GetOuter()->IsA( UPackage::StaticClass() ) &&	// Only want outer assets (these should be the only public assets, anyway)
+				( Cast<AActor>( Object ) == nullptr ) &&			// skip actors
+				( Cast<UActorComponent>( Object ) == nullptr ) &&	// skip components
+				( Cast<UClass>( Object ) == nullptr );				// skip UClasses
+
+			if( bValidObject )
+			{
+				// Ignore self referencing objects
+				if( Object != CurrentObject )
+				{
+					// ADD REFERENCE(CurrentObject DEPENDS on Object)
+					ReferencedObjects.Add( Object );
+				}
+			}
+
+			return *this;
+		}
+
+	private:
+		/** The object currently being serialized. */
+		UObject* CurrentObject;
+
+	public:
+		/** The set of referenced objects */
+		TSet< UObject* > ReferencedObjects;
+	};
+
+
+
+
+	/** Given a size in bytes and a boolean that indicates whether the size is actually known to be correct, returns a pretty
+	    string to represent that size, such as "256.0 MB", or "unknown size" */
+	static FString MakeBestSizeString( const SIZE_T SizeInBytes, const bool bHasKnownSize )
+	{
+		FString BestSizeString;			
+
+		const FNumericUnit<double> BestUnit = FUnitConversion::QuantizeUnitsToBestFit( (double)SizeInBytes, EUnit::Bytes );
+
+		if( BestUnit.Units == EUnit::Bytes )
+		{
+			// We ended up with bytes, so show a decimal number
+			BestSizeString = FString::Printf( TEXT( "%s %s" ), 
+				*FText::AsNumber( SizeInBytes ).ToString(), 
+				*LOCTEXT( "Bytes", "bytes" ).ToString() );
+		}
+		else
+		{
+			// Show a fractional number with the best possible units
+			FNumberFormattingOptions NumberFormattingOptions;
+			NumberFormattingOptions.MaximumFractionalDigits = 1;	// @todo sizemap: We could make the number of digits customizable in the UI
+			NumberFormattingOptions.MinimumFractionalDigits = 0;
+			NumberFormattingOptions.MinimumIntegralDigits = 1;
+			BestSizeString = FString::Printf( TEXT( "%s %s" ), 
+				*FText::AsNumber( BestUnit.Value, &NumberFormattingOptions ).ToString(), 
+				FUnitConversion::GetUnitDisplayString( BestUnit.Units ) );
+		}
+
+		if( !bHasKnownSize )
+		{
+			if( SizeInBytes == 0 )
+			{
+				BestSizeString = LOCTEXT( "UnknownSize", "unknown size" ).ToString();
+			}
+			else
+			{
+				BestSizeString = FString::Printf( TEXT( "%s %s" ),
+					*LOCTEXT( "UnknownSizeButAtLeastThisBig", "at least" ).ToString(),
+					*BestSizeString );
+			}
+		}
+		return BestSizeString;
+	}
+}
+
+
+void SSizeMap::GatherDependenciesRecursively( FAssetRegistryModule& AssetRegistryModule, TSharedPtr<FAssetThumbnailPool>& AssetThumbnailPool, TMap<FName, TSharedPtr<FTreeMapNodeData>>& VisitedAssetPackageNames, const TArray<FName>& AssetPackageNames, const TSharedPtr<FTreeMapNodeData>& Node, TSharedPtr<FTreeMapNodeData>& SharedRootNode, int32& NumAssetsWhichFailedToLoad )
+{
+	for( const FName AssetPackageName : AssetPackageNames )
+	{
+		// Have we already added this asset to the tree?  If so, we'll either move it to a "shared" group or (if its referenced again by the same
+		// root-level asset) ignore it
+		if( VisitedAssetPackageNames.Contains( AssetPackageName ) )
+		{ 
+			// OK, we've determined that this asset has already been referenced by something else in our tree.  We'll move it to a "shared" group
+			// so all of the assets that are referenced in multiple places can be seen together.
+			TSharedPtr<FTreeMapNodeData> ExistingNodeInTree = VisitedAssetPackageNames[ AssetPackageName ];
+
+			// Is the existing node not already under the "shared" group?  Note that it might still be (indirectly) under
+			// the "shared" group, in which case we'll still want to move it up to the root since we've figured out that it is
+			// actually shared between multiple assets which themselves may be shared
+			if( ExistingNodeInTree->Parent != SharedRootNode.Get() )
+			{
+				// Don't bother moving any of the assets at the root level into a "shared" bucket.  We're only trying to best
+				// represent the memory used when all of the root-level assets have become loaded.  It's OK if root-level assets
+				// are referenced by other assets in the set -- we don't need to indicate they are shared explicitly
+				if( ExistingNodeInTree->Parent->Parent != nullptr )
+				{
+					FTreeMapNodeData* MyParentNode = Node.Get();
+
+					// Find the existing node's root level asset node
+					FTreeMapNodeData* ExistingNodeRootLevelAssetNode = ExistingNodeInTree.Get();
+					while( ExistingNodeRootLevelAssetNode->Parent->Parent != nullptr )
+					{
+						ExistingNodeRootLevelAssetNode = ExistingNodeRootLevelAssetNode->Parent;
+					}
+
+					// If we're being referenced by another node within the same asset, no need to move it to a 'shared' group.  We only care
+					// about referenced assets that are shared between different assets
+					// @todo sizemap urgent: Doesn't work correctly with all assets (BP_Sky)
+					if( MyParentNode != ExistingNodeRootLevelAssetNode->Parent )
+					{
+						// This asset was already referenced by something else (or was in our top level list of assets to display sizes for)
+						if( !SharedRootNode.IsValid() )
+						{
+							// Find the root-most tree node
+							FTreeMapNodeData* RootNode = MyParentNode;
+							while( RootNode->Parent != nullptr )
+							{
+								RootNode = RootNode->Parent;
+							}
+
+							SharedRootNode = MakeShareable( new FTreeMapNodeData() );
+							RootNode->Children.Add( SharedRootNode );
+							SharedRootNode->Parent = RootNode;	// Keep back-pointer to parent node
+						}
+
+						// Reparent the node that we've now determined to be shared
+						ExistingNodeInTree->Parent->Children.Remove( ExistingNodeInTree );
+						SharedRootNode->Children.Add( ExistingNodeInTree );
+						ExistingNodeInTree->Parent = SharedRootNode.Get();
+					}
+				}
+			}
+		}
+		else
+		{
+			// This asset is new to us so far!  Let's add it to the tree.  Later as we descend through references, we might find that the
+			// asset is referenced by something else as well, in which case we'll pull it out and move it to a "shared" top-level box
+					
+			// Don't bother showing code references
+			const FString AssetPackageNameString = AssetPackageName.ToString();
+			if( !AssetPackageNameString.StartsWith( TEXT( "/Script/" ) ) )
+			{
+				FTreeMapNodeDataRef ChildTreeMapNode = MakeShareable( new FTreeMapNodeData() );
+				Node->Children.Add( ChildTreeMapNode );
+				ChildTreeMapNode->Parent = Node.Get();	// Keep back-pointer to parent node
+
+				VisitedAssetPackageNames.Add( AssetPackageName, ChildTreeMapNode );
+
+				FNodeSizeMapData& NodeSizeMapData = NodeSizeMapDataMap.Add( ChildTreeMapNode );
+
+				// Set some defaults for this node.  These will be used if we can't actually locate the asset.
+				// @todo sizemap urgent: We need a better indication in the UI when there are one or more missing assets.  Because missing assets have a size 
+				//    of zero, they are nearly impossible to zoom into.  At the least, we should have some Output Log spew when assets cannot be loaded
+				NodeSizeMapData.AssetData.AssetName = AssetPackageName;
+				NodeSizeMapData.AssetData.AssetClass = FName( *LOCTEXT( "MissingAsset", "MISSING!" ).ToString() );
+				NodeSizeMapData.AssetSize = 0;
+				NodeSizeMapData.bHasKnownSize = false;
+
+				// Find the asset using the asset registry
+				// @todo sizemap urgent: Asset registry is faster but possibly not as exhaustive (no PostLoad created references, etc.)  Maybe should be optional?
+				// @todo sizemap urgent: When not using asset registry to find references, all references of map files are not showing up
+				// @todo sizemap: We don't really need the asset registry given we need to load the objects to figure out their size, unless we make that AR-searchable.
+				//   ---> This would allow us to not have to wait for AR initialization.  But if we made size AR-searchable, we could run very quickly for large data sets!
+				// @todo sizemap urgent: When loading dependencies using the FAssetReferrenceFinder, we never find out about any missing references (unlike with the asset
+				//    registry approach, where we can tell when something is missing.)  This could product misleading results!
+				const bool bUseAssetRegistryForDependencies = true;
+
+				const FString AssetPathString = AssetPackageNameString + TEXT(".") + FPackageName::GetLongPackageAssetName( AssetPackageNameString );
+				const FAssetData FoundAssetData = AssetRegistryModule.Get().GetAssetByObjectPath( FName( *AssetPathString ) );
+				if( FoundAssetData.IsValid() )
+				{
+					NodeSizeMapData.AssetData = FoundAssetData;
+
+					// Now actually load up the asset.  We need it in memory in order to accurately determine its size.
+					// @todo sizemap: We could async load these packages to make the editor experience a bit nicer (smoother progress)
+					UObject* Asset = StaticLoadObject( UObject::StaticClass(), nullptr, *AssetPathString );
+					if( Asset != nullptr )
+					{
+						TArray<FName> ReferencedAssetPackageNames;
+						if( bUseAssetRegistryForDependencies )
+						{
+							AssetRegistryModule.Get().GetDependencies( AssetPackageName, ReferencedAssetPackageNames );
+						}
+						else
+						{
+							SizeMapInternals::FAssetReferenceFinder References( Asset );
+							for( UObject* Object : References.ReferencedObjects )
+							{
+								ReferencedAssetPackageNames.Add( FName( *Object->GetOutermost()->GetPathName() ) );
+							}
+						}
+
+						NodeSizeMapData.AssetSize = Asset->GetResourceSize( EResourceSizeMode::Exclusive );
+						NodeSizeMapData.bHasKnownSize = ( NodeSizeMapData.AssetSize != UObject::RESOURCE_SIZE_NONE );
+						if( !NodeSizeMapData.bHasKnownSize )
+						{
+							// Asset has no meaningful size
+							NodeSizeMapData.AssetSize = 0;
+						}
+
+								
+						// Now visit all of the assets that we are referencing
+						GatherDependenciesRecursively( AssetRegistryModule, AssetThumbnailPool, VisitedAssetPackageNames, ReferencedAssetPackageNames, ChildTreeMapNode, SharedRootNode, NumAssetsWhichFailedToLoad );
+
+					}
+					else
+					{
+						++NumAssetsWhichFailedToLoad;
+					}
+				}
+				else
+				{
+					++NumAssetsWhichFailedToLoad;
+				}
+			}
+		}
+	}
+}
+
+	
+void SSizeMap::FinalizeNodesRecursively( TSharedPtr<FTreeMapNodeData>& Node, const TSharedPtr<FTreeMapNodeData>& SharedRootNode, int32& TotalAssetCount, SIZE_T& TotalSize, bool& bAnyUnknownSizes )
+{
+	// Process children first, so we can get the totals for the root node and shared nodes
+	int32 SubtreeAssetCount = 0;
+	SIZE_T SubtreeSize = 0;
+	bool bAnyUnknownSizesInSubtree = false;
+	{
+		for( TSharedPtr<FTreeMapNodeData> ChildNode : Node->Children )
+		{
+			FinalizeNodesRecursively( ChildNode, SharedRootNode, SubtreeAssetCount, SubtreeSize, bAnyUnknownSizesInSubtree );
+		}
+
+		TotalAssetCount += SubtreeAssetCount;
+		TotalSize += SubtreeSize;
+		if( bAnyUnknownSizesInSubtree )
+		{
+			bAnyUnknownSizes = true;
+		}
+	}
+
+	if( Node == SharedRootNode )
+	{
+		// @todo sizemap: Should we indicate in a non-shared parent node how many if its dependents ended up being in the "shared" bucket?  Probably 
+		// not that important, because the user can choose to view that asset in isolation to see the full tree.
+		Node->Name = FString::Printf( TEXT( "%s (%s)" ),
+			*LOCTEXT( "SharedGroupName", "*SHARED*" ).ToString(),
+			*SizeMapInternals::MakeBestSizeString( SubtreeSize, !bAnyUnknownSizes ) );
+
+		// Container nodes are always auto-sized
+		Node->Size = 0.0f;
+	}
+	else if( Node->Parent == nullptr )
+	{
+		// Tree root is always auto-sized
+		Node->Size = 0.0f;
+	}
+	else
+	{
+		const FNodeSizeMapData& NodeSizeMapData = NodeSizeMapDataMap.FindChecked( Node.ToSharedRef() );
+
+		++TotalAssetCount;
+		TotalSize += NodeSizeMapData.AssetSize;
+
+		if( !NodeSizeMapData.bHasKnownSize )
+		{
+			bAnyUnknownSizes = true;
+		}
+
+		// Setup a thumbnail
+		const FSlateBrush* DefaultThumbnailSlateBrush;
+		{
+			// For non-class types, use the default based upon the actual asset class
+			// This has the side effect of not showing a class icon for assets that don't have a proper thumbnail image available
+			bool bIsClassType = false;
+			const UClass* ThumbnailClass = FClassIconFinder::GetIconClassForAssetData( NodeSizeMapData.AssetData, &bIsClassType );
+			const FName DefaultThumbnail = (bIsClassType) ? NAME_None : FName(*FString::Printf(TEXT("ClassThumbnail.%s"), *NodeSizeMapData.AssetData.AssetClass.ToString()));							
+			DefaultThumbnailSlateBrush = FClassIconFinder::FindThumbnailForClass(ThumbnailClass, DefaultThumbnail);
+
+			// @todo sizemap urgent: Actually implement rendered thumbnail support, not just class-based background images
+
+			// const int32 ThumbnailSize = 128;	// @todo sizemap: Hard-coded thumbnail size.  Move this elsewhere
+			//	TSharedRef<FAssetThumbnail> AssetThumbnail( new FAssetThumbnail( NodeSizeMapData.AssetData, ThumbnailSize, ThumbnailSize, AssetThumbnailPool ) );
+			//	ChildTreeMapNode->AssetThumbnail = AssetThumbnail->MakeThumbnailImage();
+		}
+
+		if( Node->IsLeafNode() )
+		{
+			Node->CenterText = SizeMapInternals::MakeBestSizeString( NodeSizeMapData.AssetSize, NodeSizeMapData.bHasKnownSize );
+
+			Node->Size = NodeSizeMapData.AssetSize;
+
+			// The STreeMap widget is not expecting zero-sized leaf nodes.  So we make them very small instead.
+			if( Node->Size == 0 )
+			{
+				Node->Size = 1;
+			}
+
+			// Leaf nodes get a background picture
+			Node->BackgroundBrush = DefaultThumbnailSlateBrush;
+
+			// "Asset name"
+			// "Asset type"
+			Node->Name = NodeSizeMapData.AssetData.AssetName.ToString();
+			Node->Name2 = NodeSizeMapData.AssetData.AssetClass.ToString();
+		}
+		else
+		{
+			// Container nodes are always auto-sized
+			Node->Size = 0.0f;
+
+			const bool bNeedsSelfNode = NodeSizeMapData.AssetSize > 0;
+			if( bNeedsSelfNode )
+			{
+				// "Asset name" (size)
+				Node->Name = FString::Printf( TEXT( "%s (%s)" ),
+					*NodeSizeMapData.AssetData.AssetName.ToString(),
+					*SizeMapInternals::MakeBestSizeString( SubtreeSize + NodeSizeMapData.AssetSize, !bAnyUnknownSizesInSubtree && NodeSizeMapData.bHasKnownSize ) );
+
+				// We have children, so make some space for our own asset's size within our box
+				FTreeMapNodeDataRef ChildSelfTreeMapNode = MakeShareable( new FTreeMapNodeData() );
+				Node->Children.Add( ChildSelfTreeMapNode );
+				ChildSelfTreeMapNode->Parent = Node.Get();	// Keep back-pointer to parent node
+
+				// Map the "self" node to the same node data as its parent
+				NodeSizeMapDataMap.Add( ChildSelfTreeMapNode, NodeSizeMapData );
+
+				// "*SELF*"
+				// "Asset type"
+				ChildSelfTreeMapNode->Name = LOCTEXT( "SelfNodeLabel", "*SELF*" ).ToString();
+				ChildSelfTreeMapNode->Name2 = NodeSizeMapData.AssetData.AssetClass.ToString();
+
+				ChildSelfTreeMapNode->CenterText = SizeMapInternals::MakeBestSizeString( NodeSizeMapData.AssetSize, NodeSizeMapData.bHasKnownSize );
+				ChildSelfTreeMapNode->Size = NodeSizeMapData.AssetSize;
+
+				// Leaf nodes get a background picture
+				ChildSelfTreeMapNode->BackgroundBrush = DefaultThumbnailSlateBrush;
+			}
+			else
+			{
+				// "Asset name (asset type, size)"
+				Node->Name = FString::Printf( TEXT( "%s (%s, %s)" ),
+					*NodeSizeMapData.AssetData.AssetName.ToString(),
+					*NodeSizeMapData.AssetData.AssetClass.ToString(),
+					*SizeMapInternals::MakeBestSizeString( SubtreeSize + NodeSizeMapData.AssetSize, !bAnyUnknownSizesInSubtree && NodeSizeMapData.bHasKnownSize ) );
+			}
+		}
+
+	}
+
+	// Sort all of my child nodes alphabetically.  This is just so that we get deterministic results when viewing the
+	// same sets of assets.
+	Node->Children.StableSort(
+			[]( const FTreeMapNodeDataPtr& A, const FTreeMapNodeDataPtr& B )
+			{
+				return A->Name < B->Name;
+			}
+		);
+}
+
+	
+void SSizeMap::RefreshMap()
+{
+	// Wipe the current tree out
+	RootTreeMapNode->Children.Empty();
+	NodeSizeMapDataMap.Empty();
+
+	
+	// First, do a pass to gather asset dependencies and build up a tree
+	TMap<FName, TSharedPtr<FTreeMapNodeData>> VisitedAssetPackageNames;
+	TSharedPtr<FTreeMapNodeData> SharedRootNode;
+	int32 NumAssetsWhichFailedToLoad = 0;
+	FAssetRegistryModule& AssetRegistryModule = FModuleManager::LoadModuleChecked<FAssetRegistryModule>("AssetRegistry");
+	GatherDependenciesRecursively( AssetRegistryModule, AssetThumbnailPool, VisitedAssetPackageNames, RootAssetPackageNames, RootTreeMapNode, SharedRootNode, NumAssetsWhichFailedToLoad );
+
+
+	// Next, do another pass over our tree to and count how big the assets are and to set the node labels.  Also in this pass, we may
+	// create some additional "self" nodes for assets that have children but also take up size themselves.
+	int32 TotalAssetCount = 0;
+	SIZE_T TotalSize = 0;
+	bool bAnyUnknownSizes = false;
+	FinalizeNodesRecursively( RootTreeMapNode, SharedRootNode, TotalAssetCount, TotalSize, bAnyUnknownSizes );
+
+
+	// Create a nice name for the tree!
+	if( NumAssetsWhichFailedToLoad > 0 )
+	{
+		RootTreeMapNode->Name = FString::Printf( TEXT( "%s %i %s" ),
+			*LOCTEXT( "RootNode_WarningPrefix", "WARNING:" ).ToString(),
+			NumAssetsWhichFailedToLoad,
+			*LOCTEXT( "RootNode_NAssetsFailedToLoad", "assets were missing!  Only partial results shown." ).ToString() );
+	}
+	else if( RootAssetPackageNames.Num() == 1 && !SharedRootNode.IsValid() )
+	{
+		// @todo sizemap: When zoomed right into one asset, can we use the Class color for the node instead of grey?
+		
+		// The root will only have one child, so go ahead and use that child as the actual root
+		FTreeMapNodeDataPtr OnlyChild = RootTreeMapNode->Children[0];
+		OnlyChild->CopyNodeInto( *RootTreeMapNode );
+		RootTreeMapNode->Children = OnlyChild->Children;
+		RootTreeMapNode->Parent = nullptr;
+		for( const auto& ChildNode : RootTreeMapNode->Children )
+		{
+			ChildNode->Parent = RootTreeMapNode.Get();
+		}
+
+		// Use a more descriptive name for the root level node
+		RootTreeMapNode->Name = FString::Printf( TEXT( "%s %s (%i %s)" ),
+			*LOCTEXT( "RootNode_SizeMapForOneAsset", "Size map for" ).ToString(),
+			*OnlyChild->Name,
+			TotalAssetCount,
+			*LOCTEXT( "RootNode_References", "total assets" ).ToString() );
+	}
+	else
+	{
+		// Multiple assets (or at least some shared assets) at the root level
+		RootTreeMapNode->BackgroundBrush = nullptr;
+		RootTreeMapNode->Size = 0.0f;
+		RootTreeMapNode->Parent = nullptr;
+		RootTreeMapNode->Name = FString::Printf( TEXT( "%s %i %s (%i %s, %s)" ),
+			*LOCTEXT( "RootNode_SizeMapForMultiple", "Size map for" ).ToString(),
+			RootAssetPackageNames.Num(),
+			*LOCTEXT( "RootNode_Assets", "assets" ).ToString(),
+			TotalAssetCount,
+			*LOCTEXT( "RootNode_References", "total assets" ).ToString(),
+			*SizeMapInternals::MakeBestSizeString( TotalSize, !bAnyUnknownSizes ) );
+	}
+
+						
+	// OK, now refresh the actual tree map widget so our new tree will be displayed.
+	const bool bShouldPlayTransition = false;
+	TreeMapWidget->RebuildTreeMap( bShouldPlayTransition );
+}
+
+
+void SSizeMap::OnInitialAssetRegistrySearchComplete()
+{
+	RefreshMap();
+}
+
+
+void SSizeMap::OnTreeMapNodeDoubleClicked( FTreeMapNodeData& TreeMapNodeData )
+{
+	const FNodeSizeMapData* NodeSizeMapData = NodeSizeMapDataMap.Find( TreeMapNodeData.AsShared() );
+	if( NodeSizeMapData != nullptr )
+	{
+		TArray<FAssetData> Assets;
+		Assets.Add( NodeSizeMapData->AssetData );
+		GEditor->SyncBrowserToObjects( Assets );
+	}
+}
+
+
+
+// @todo sizemap urgent: We should add a spinner while we discover and load assets
+// @todo sizemap urgent: Audit common asset types and make sure they have a useful GetResourceSize() implementation
+//		-> Some implementations are including the size of editor-only data (Static Mesh).  This should be configurable!
+// @todo sizemap urgent: It would be great to be able to see 0-sized/unknown-sized/tiny-sized nodes somehow, or at least a count of them
+//		- You'd almost want a way to zoom in super small.  Or a little laserpointer effect that shows labels for super-tiny nodes
+// @todo sizemap urgent: Need inline help to figure out mousewheel?  Single click currently does nothing.  
+// @todo sizemap urgent: When an asset is not reporting a size, this tool should help to escalate that by making it very obvious that its missing (no GetResourceSize function)
+// @todo sizemap: Can "zoom into" <Self> nodes which is a bit weird.  (Asset name no longer draws).  Maybe disallow zooming all the way in?
+//  --> Double-click current zooms directly into a single asset. Not super useful.  Maybe only allow zooming down to the deepest nodes, not leaves
+// @todo sizemap: When the tab is restored from layout, it will be totally empty.  Instead it should probably have instructions for how to show the sizes for assets. (Same with Reference Viewer)
+// @todo sizemap: We ideally want to replace the SReferenceTree code with SSizeMap
+// @todo sizemap: Add a tree view that shows all of the references along with their sizes (so you can see Very Small references)
+// @todo sizemap: It would be useful to be able to preview the sizes for specific platforms? (option in the UI)
+// @todo sizemap: Should we show the percentage of total size as an option in the UI (though, the size of the boxes show this pretty well.)
+// @todo sizemap: Should we show the folder part of the asset name as a tool-tip?
+// @todo sizemap: It might be nice to unload assets that were loaded by us after the size map is built up
+// @todo sizemap urgent: Trying size map for all engine content BROKE the window updating
+
+
+#undef LOCTEXT_NAMESPACE
\ No newline at end of file
diff --git a/Engine/Source/Editor/SizeMap/SSizeMap.h b/Engine/Source/Editor/SizeMap/SSizeMap.h
new file mode 100644
--- /dev/null
+++ b/Engine/Source/Editor/SizeMap/SSizeMap.h
@@ -0,0 +1,83 @@
+// Copyright 1998-2015 Epic Games, Inc. All Rights Reserved.
+
+#pragma once
+
+/**
+ * Tree map for displaying the size of assets
+ */
+class SSizeMap : public SCompoundWidget
+{
+
+public:
+
+	SLATE_BEGIN_ARGS( SSizeMap ) {}
+	SLATE_END_ARGS()
+
+	/** Default constructor for SSizeMap */
+	SSizeMap();
+
+	/** Destructor for SSizeMap */
+	~SSizeMap();
+
+	/**
+	 * Construct the widget
+	 *
+	 * @param	InArgs				A declaration from which to construct the widget
+	 */
+	void Construct( const FArguments& InArgs );
+
+	/** Sets the assets to view at the root of the size map.  This will rebuild the map. */
+	void SetRootAssetPackageNames( const TArray<FName>& NewRootAssetPackageNames );
+
+
+protected:
+
+	/** Called after the initial asset registry scan finishes */
+	void OnInitialAssetRegistrySearchComplete();
+
+	/** Recursively discovers and loads dependent assets, building up a tree map node hierarchy */
+	void GatherDependenciesRecursively( class FAssetRegistryModule& AssetRegistryModule, TSharedPtr<class FAssetThumbnailPool>& AssetThumbnailPool, TMap<FName, TSharedPtr<class FTreeMapNodeData>>& VisitedAssetPackageNames, const TArray<FName>& AssetPackageNames, const TSharedPtr<FTreeMapNodeData>& ParentTreeMapNode, TSharedPtr<FTreeMapNodeData>& SharedRootNode, int32& NumAssetsWhichFailedToLoad );
+
+	/** After the node tree is built up, this function is called to generate nice labels for the nodes and to do a final clean-up pass on the tree */
+	void FinalizeNodesRecursively( TSharedPtr<FTreeMapNodeData>& Node, const TSharedPtr<FTreeMapNodeData>& SharedRootNode, int32& TotalAssetCount, SIZE_T& TotalSize, bool& bAnyUnknownSizes );
+
+	/** Refreshes the display */
+	void RefreshMap();
+
+	/** Called when the user double-clicks on an asset in the tree */
+	void OnTreeMapNodeDoubleClicked( class FTreeMapNodeData& TreeMapNodeData );
+
+
+protected:
+
+	/** Our tree map widget */
+	TSharedPtr<class STreeMap> TreeMapWidget;
+
+	/** The assets we were asked to look at */
+	TArray<FName> RootAssetPackageNames;
+
+	/** Our tree map source data */
+	TSharedPtr<class FTreeMapNodeData> RootTreeMapNode;
+
+	/** Thumbnail pool */
+	TSharedPtr<class FAssetThumbnailPool> AssetThumbnailPool;
+
+
+	/** This struct will contain size map-specific payload data that we'll associate with tree map nodes using a hash table */
+	struct FNodeSizeMapData
+	{
+		/** How big the asset is */
+		SIZE_T AssetSize;
+
+		/** Whether it has a known size or not */
+		bool bHasKnownSize;
+
+		/** Data from the asset registry about this asset */
+		FAssetData AssetData;
+	};
+
+	/** Maps a tree node to our size map-specific user data for that node */
+	typedef TMap<TSharedRef<FTreeMapNodeData>, FNodeSizeMapData> FNodeSizeMapDataMap;
+	FNodeSizeMapDataMap NodeSizeMapDataMap;
+};
+
diff --git a/Engine/Source/Editor/SizeMap/SizeMap.Build.cs b/Engine/Source/Editor/SizeMap/SizeMap.Build.cs
new file mode 100644
--- /dev/null
+++ b/Engine/Source/Editor/SizeMap/SizeMap.Build.cs
@@ -0,0 +1,36 @@
+// Copyright 1998-2015 Epic Games, Inc. All Rights Reserved.
+
+namespace UnrealBuildTool.Rules
+{
+	public class SizeMap : ModuleRules
+	{
+        public SizeMap(TargetInfo Target)
+		{
+			PrivateDependencyModuleNames.AddRange(
+                new string[] {
+				    "Core",
+				    "CoreUObject",
+				    "Engine",
+                    "InputCore",
+				    "Slate",
+					"SlateCore",
+                    "EditorStyle",
+				    "UnrealEd",
+					"TreeMap"
+			    }
+			);
+
+            PrivateIncludePathModuleNames.AddRange(
+                new string[] {
+				    "AssetRegistry",
+			    }
+            );
+
+            DynamicallyLoadedModuleNames.AddRange(
+                new string[] {
+				    "AssetRegistry",
+			    }
+            );
+		}
+	}
+}
\ No newline at end of file
diff --git a/Engine/Source/Editor/SizeMap/SizeMapActions.cpp b/Engine/Source/Editor/SizeMap/SizeMapActions.cpp
new file mode 100644
--- /dev/null
+++ b/Engine/Source/Editor/SizeMap/SizeMapActions.cpp
@@ -0,0 +1,14 @@
+// Copyright 1998-2015 Epic Games, Inc. All Rights Reserved.
+
+#include "SizeMapModule.h"
+#include "SizeMapActions.h"
+
+#define LOCTEXT_NAMESPACE ""
+
+void FSizeMapActions::RegisterCommands()
+{
+	// @todo sizemap: Register any commands we have here
+	// UI_COMMAND(DoSomething, "Do Something", "Did something tooltip.", EUserInterfaceActionType::Button, FInputChord());
+}
+
+#undef LOCTEXT_NAMESPACE
\ No newline at end of file
diff --git a/Engine/Source/Editor/SizeMap/SizeMapActions.h b/Engine/Source/Editor/SizeMap/SizeMapActions.h
new file mode 100644
--- /dev/null
+++ b/Engine/Source/Editor/SizeMap/SizeMapActions.h
@@ -0,0 +1,26 @@
+// Copyright 1998-2015 Epic Games, Inc. All Rights Reserved.
+
+#pragma once
+
+// Actions that can be invoked in the Size Map
+class FSizeMapActions : public TCommands<FSizeMapActions>
+{
+public:
+	FSizeMapActions() : TCommands<FSizeMapActions>
+	(
+		"SizeMap",
+		NSLOCTEXT("Contexts", "SizeMap", "Size Map"),
+		"MainFrame", FEditorStyle::GetStyleSetName()
+	)
+	{
+	}
+
+	// TCommands<> interface
+	virtual void RegisterCommands() override;
+	// End of TCommands<> interface
+public:
+
+	// @todo sizemap: Register any commands we have here
+	// TSharedPtr<FUICommandInfo> DoSomething;
+	
+};
diff --git a/Engine/Source/Editor/SizeMap/SizeMapModule.cpp b/Engine/Source/Editor/SizeMap/SizeMapModule.cpp
new file mode 100644
--- /dev/null
+++ b/Engine/Source/Editor/SizeMap/SizeMapModule.cpp
@@ -0,0 +1,56 @@
+// Copyright 1998-2015 Epic Games, Inc. All Rights Reserved.
+
+#include "SizeMapModule.h"
+#include "ISizeMapModule.h"
+#include "SSizeMap.h"
+#include "SDockTab.h"
+
+#define LOCTEXT_NAMESPACE "SizeMap"
+
+class FSizeMapModule : public ISizeMapModule
+{
+public:
+	FSizeMapModule()
+		: SizeMapTabId("SizeMap")
+	{
+	}
+
+	virtual void StartupModule() override
+	{
+		FGlobalTabmanager::Get()->RegisterNomadTabSpawner(SizeMapTabId, FOnSpawnTab::CreateRaw( this, &FSizeMapModule::SpawnSizeMapTab ) )
+			.SetDisplayName( LOCTEXT("SizeMapTitle", "Size Map") )
+			.SetMenuType( ETabSpawnerMenuType::Hidden );
+	}
+
+	virtual void ShutdownModule() override
+	{
+		FGlobalTabmanager::Get()->UnregisterNomadTabSpawner(SizeMapTabId);
+	}
+
+	virtual void InvokeSizeMapTab(const TArray<FName>& AssetPackageNames) override
+	{
+		TSharedRef<SDockTab> NewTab = FGlobalTabmanager::Get()->InvokeTab( SizeMapTabId );
+		TSharedRef<SSizeMap> SizeMap = StaticCastSharedRef<SSizeMap>( NewTab->GetContent() );
+		SizeMap->SetRootAssetPackageNames( AssetPackageNames );
+	}
+
+private:
+	TSharedRef<SDockTab> SpawnSizeMapTab( const FSpawnTabArgs& SpawnTabArgs )
+	{
+		TSharedRef<SDockTab> NewTab = SNew(SDockTab)
+			.TabRole(ETabRole::NomadTab);
+
+		NewTab->SetContent( SNew(SSizeMap) );
+
+		return NewTab;
+	}
+
+private:
+	FName SizeMapTabId;
+};
+
+
+IMPLEMENT_MODULE( FSizeMapModule, SizeMap )
+
+
+#undef LOCTEXT_NAMESPACE
\ No newline at end of file
diff --git a/Engine/Source/Editor/SizeMap/SizeMapModule.h b/Engine/Source/Editor/SizeMap/SizeMapModule.h
new file mode 100644
--- /dev/null
+++ b/Engine/Source/Editor/SizeMap/SizeMapModule.h
@@ -0,0 +1,6 @@
+// Copyright 1998-2015 Epic Games, Inc. All Rights Reserved.
+
+#pragma once
+
+#include "UnrealEd.h"
+#include "ISizeMapModule.h"
diff --git a/Engine/Source/Editor/UnrealEd/Private/EditorEngine.cpp b/Engine/Source/Editor/UnrealEd/Private/EditorEngine.cpp
--- a/Engine/Source/Editor/UnrealEd/Private/EditorEngine.cpp
+++ b/Engine/Source/Editor/UnrealEd/Private/EditorEngine.cpp
@@ -643,7 +643,10 @@ void UEditorEngine::Init(IEngineLoop* InEngineLoop)
 			TEXT("SourceCodeAccess"),
 			TEXT("BehaviorTreeEditor"),
 			TEXT("HardwareTargeting"),
-			TEXT("LocalizationDashboard")
+			TEXT("LocalizationDashboard"),
+			TEXT("ReferenceViewer"),
+			TEXT("TreeMap"),
+			TEXT("SizeMap")
 		};
 
 		FScopedSlowTask ModuleSlowTask(ARRAY_COUNT(ModuleNames));
diff --git a/Engine/Source/Editor/UnrealEd/Private/Toolkits/AssetEditorToolkit.cpp b/Engine/Source/Editor/UnrealEd/Private/Toolkits/AssetEditorToolkit.cpp
--- a/Engine/Source/Editor/UnrealEd/Private/Toolkits/AssetEditorToolkit.cpp
+++ b/Engine/Source/Editor/UnrealEd/Private/Toolkits/AssetEditorToolkit.cpp
@@ -12,6 +12,7 @@
 #include "IUserFeedbackModule.h"
 #include "IDocumentation.h"
 #include "ReferenceViewer.h"
+#include "ISizeMapModule.h"
 #include "IIntroTutorials.h"
 #include "SuperSearchModule.h"
 #include "SDockTab.h"
@@ -200,6 +201,11 @@ void FAssetEditorToolkit::InitAssetEditor( const EToolkitMode::Type Mode, const
 		FExecuteAction::CreateSP( this, &FAssetEditorToolkit::ViewReferences_Execute ),
 		FCanExecuteAction::CreateSP( this, &FAssetEditorToolkit::CanViewReferences ));
 	
+	ToolkitCommands->MapAction(
+		FGlobalEditorCommonCommands::Get().ViewSizeMap,
+		FExecuteAction::CreateSP( this, &FAssetEditorToolkit::ViewSizeMap_Execute ),
+		FCanExecuteAction::CreateSP( this, &FAssetEditorToolkit::CanViewSizeMap ));
+	
 	ToolkitCommands->MapAction(
 		FGlobalEditorCommonCommands::Get().OpenDocumentation,
 		FExecuteAction::CreateSP( this, &FAssetEditorToolkit::BrowseDocumentation_Execute ) );
@@ -578,6 +584,29 @@ bool FAssetEditorToolkit::CanViewReferences()
 	return ViewableObjects.Num() > 0;
 }
 
+void FAssetEditorToolkit::ViewSizeMap_Execute()
+{
+	if (ensure( ViewableObjects.Num() > 0))
+	{
+		ISizeMapModule::Get().InvokeSizeMapTab(ViewableObjects);
+	}
+}
+
+bool FAssetEditorToolkit::CanViewSizeMap()
+{
+	ViewableObjects.Empty();
+	for (const auto EditingObject : EditingObjects)
+	{
+		// Don't allow user to perform certain actions on objects that aren't actually assets (e.g. Level Script blueprint objects)
+		if (EditingObject != NULL && EditingObject->IsAsset())
+		{
+			ViewableObjects.Add(EditingObject->GetOuter()->GetFName());
+		}
+	}
+
+	return ViewableObjects.Num() > 0;
+}
+
 FString FAssetEditorToolkit::GetDocumentationLink() const
 {
 	return FString(TEXT("%ROOT%"));
@@ -699,6 +728,7 @@ void FAssetEditorToolkit::FillDefaultAssetMenuCommands( FMenuBuilder& MenuBuilde
 	{
 		MenuBuilder.AddMenuEntry( FGlobalEditorCommonCommands::Get().FindInContentBrowser, "FindInContentBrowser", LOCTEXT("FindInContentBrowser", "Find in Content Browser...") );
 		MenuBuilder.AddMenuEntry( FGlobalEditorCommonCommands::Get().ViewReferences);
+		MenuBuilder.AddMenuEntry( FGlobalEditorCommonCommands::Get().ViewSizeMap);
 
 		// Add a reimport menu entry for each supported editable object
 		for( auto ObjectIter = EditingObjects.CreateConstIterator(); ObjectIter; ++ObjectIter )
diff --git a/Engine/Source/Editor/UnrealEd/Private/Toolkits/GlobalEditorCommonCommands.cpp b/Engine/Source/Editor/UnrealEd/Private/Toolkits/GlobalEditorCommonCommands.cpp
--- a/Engine/Source/Editor/UnrealEd/Private/Toolkits/GlobalEditorCommonCommands.cpp
+++ b/Engine/Source/Editor/UnrealEd/Private/Toolkits/GlobalEditorCommonCommands.cpp
@@ -27,7 +27,8 @@ void FGlobalEditorCommonCommands::RegisterCommands()
 	UI_COMMAND( SummonOpenAssetDialog, "Open Asset...", "Summons an asset picker", EUserInterfaceActionType::Button, FInputChord(EModifierKey::Control, EKeys::P) );
 	UI_COMMAND( SummonOpenAssetDialogAlternate, "Open Asset...", "Summons an asset picker", EUserInterfaceActionType::Button, FInputChord(EModifierKey::Alt | EModifierKey::Shift, EKeys::O));
 	UI_COMMAND( FindInContentBrowser, "Find in Content Browser", "Summons the Content Browser and navigates to the selected asset", EUserInterfaceActionType::Button, FInputChord(EModifierKey::Control, EKeys::B));
-	UI_COMMAND( ViewReferences, "View References", "Launches the reference viewer showing the selected assets' references", EUserInterfaceActionType::Button, FInputChord(EModifierKey::Shift | EModifierKey::Alt, EKeys::R));
+	UI_COMMAND( ViewReferences, "Reference Viewer...", "Launches the reference viewer showing the selected assets' references", EUserInterfaceActionType::Button, FInputChord(EModifierKey::Shift | EModifierKey::Alt, EKeys::R));
+	UI_COMMAND( ViewSizeMap, "Size Map...", "Displays an interactive map showing the approximate size of this asset and everything it references", EUserInterfaceActionType::Button, FInputChord(EModifierKey::Shift | EModifierKey::Alt, EKeys::M));	// @todo sizemap: Make sure key is not used already
 	
 	UI_COMMAND( OpenConsoleCommandBox, "Open Console Command Box", "Opens an edit box where you can type in a console command", EUserInterfaceActionType::Button, FInputChord(EKeys::Tilde));
 
diff --git a/Engine/Source/Editor/UnrealEd/Public/Toolkits/AssetEditorToolkit.h b/Engine/Source/Editor/UnrealEd/Public/Toolkits/AssetEditorToolkit.h
--- a/Engine/Source/Editor/UnrealEd/Public/Toolkits/AssetEditorToolkit.h
+++ b/Engine/Source/Editor/UnrealEd/Public/Toolkits/AssetEditorToolkit.h
@@ -232,6 +232,12 @@ private:
 	/** If true ViewReferences_Execute can be called, also caches ViewableObjects */
 	bool CanViewReferences();
 
+	/** Called when "View Size Map" is called for this asset */
+	void ViewSizeMap_Execute();
+
+	/** If true ViewSizeMap_Execute can be called, also caches ViewableObjects */
+	bool CanViewSizeMap();
+
 protected:
 
 	/** For standalone asset editing tool-kits that were switched from world-centric mode on the fly, this stores
diff --git a/Engine/Source/Editor/UnrealEd/Public/Toolkits/GlobalEditorCommonCommands.h b/Engine/Source/Editor/UnrealEd/Public/Toolkits/GlobalEditorCommonCommands.h
--- a/Engine/Source/Editor/UnrealEd/Public/Toolkits/GlobalEditorCommonCommands.h
+++ b/Engine/Source/Editor/UnrealEd/Public/Toolkits/GlobalEditorCommonCommands.h
@@ -28,6 +28,7 @@ public:
 	TSharedPtr<FUICommandInfo> SummonOpenAssetDialogAlternate;
 	TSharedPtr<FUICommandInfo> OpenDocumentation;
 	TSharedPtr<FUICommandInfo> ViewReferences;
+	TSharedPtr<FUICommandInfo> ViewSizeMap;
 	TSharedPtr<FUICommandInfo> OpenConsoleCommandBox;
 };
 
diff --git a/Engine/Source/Editor/UnrealEd/UnrealEd.Build.cs b/Engine/Source/Editor/UnrealEd/UnrealEd.Build.cs
--- a/Engine/Source/Editor/UnrealEd/UnrealEd.Build.cs
+++ b/Engine/Source/Editor/UnrealEd/UnrealEd.Build.cs
@@ -54,6 +54,7 @@ public class UnrealEd : ModuleRules
 				"OutputLog",
 				"Landscape",
                 "Niagara",
+				"SizeMap",
 			}
 		);
 
@@ -127,6 +128,7 @@ public class UnrealEd : ModuleRules
 				"MeshPaint",
                 "Foliage",
                 "VectorVM",
+				"TreeMap",
 			}
 		);
 
@@ -198,6 +200,7 @@ public class UnrealEd : ModuleRules
 				"HotReload",
                 "IOSPlatformEditor",
 				"HTML5PlatformEditor",
+				"SizeMap",
 			}
 		);
 
