Fix for (some) shadowing issues in the editor.

diff --git a/Engine/Source/Developer/TreeMap/STreeMap.cpp b/Engine/Source/Developer/TreeMap/STreeMap.cpp
--- a/Engine/Source/Developer/TreeMap/STreeMap.cpp
+++ b/Engine/Source/Developer/TreeMap/STreeMap.cpp
@@ -415,14 +415,14 @@ int32 STreeMap::OnPaint( const FPaintArgs& Args, const FGeometry& AllottedGeomet
 								if( !BlendedVisual.NodeData->CenterText.IsEmpty() && BlendedVisual.NodeData->IsLeafNode() )
 								{
 									// If the visual is smaller than the text we're going to draw, try using a smaller font
-									const FSlateFontInfo* CenterTextFont = &BlendedVisual.CenterTextFont;
-									const FVector2D FullTextSize = FontMeasureService->Measure( BlendedVisual.NodeData->CenterText, *CenterTextFont );
+									const FSlateFontInfo* BlendedVisualCenterTextFont = &BlendedVisual.CenterTextFont;
+									const FVector2D FullTextSize = FontMeasureService->Measure( BlendedVisual.NodeData->CenterText, *BlendedVisualCenterTextFont );
 									if( ScreenSpaceVisualSize.X < FullTextSize.X || ScreenSpaceVisualSize.Y < FullTextSize.Y )	// @todo treemap: assumption that NameFont will be smaller than CenterText font
 									{
-										CenterTextFont = &BlendedVisual.NameFont;
+										BlendedVisualCenterTextFont = &BlendedVisual.NameFont;
 									}
 
-									const FVector2D TextSize = FontMeasureService->Measure( BlendedVisual.NodeData->CenterText, *CenterTextFont );
+									const FVector2D TextSize = FontMeasureService->Measure( BlendedVisual.NodeData->CenterText, *BlendedVisualCenterTextFont );
 									const auto TextX = VisualPosition.X + FMath::Max( ChildContainerTextPadding, (float)BlendedVisual.Size.X * 0.5f - TextSize.X * 0.5f );	// Clamp to left edge if cropped, so the user can at least read the beginning
 									const auto TextY = VisualPosition.Y + FMath::Max( ChildContainerTextPadding + NameTextHeight + Name2TextHeight, (float)BlendedVisual.Size.Y * 0.5f - TextSize.Y * 0.5f );	// Clamp to bottom of first line of text if cropped
 
@@ -432,7 +432,7 @@ int32 STreeMap::OnPaint( const FPaintArgs& Args, const FGeometry& AllottedGeomet
 										LayerId,
 										AllottedGeometry.ToOffsetPaintGeometry( ShadowOffset + FVector2D( TextX, TextY ) ),
 										BlendedVisual.NodeData->CenterText,
-										*CenterTextFont,
+										*BlendedVisualCenterTextFont,
 										TextClippingRect,
 										DrawEffects,
 										InWidgetStyle.GetColorAndOpacityTint() * VisualTextColor );
diff --git a/Engine/Source/Editor/PropertyEditor/Private/PropertyEditorHelpers.cpp b/Engine/Source/Editor/PropertyEditor/Private/PropertyEditorHelpers.cpp
--- a/Engine/Source/Editor/PropertyEditor/Private/PropertyEditorHelpers.cpp
+++ b/Engine/Source/Editor/PropertyEditor/Private/PropertyEditorHelpers.cpp
@@ -757,8 +757,8 @@ namespace PropertyEditorHelpers
 
 		TWeakPtr<FPropertyNode> WeakPropertyEditor = PropertyEditor->GetPropertyNode();
 
-		TAttribute<bool>::FGetter EnabledDelegate = TAttribute<bool>::FGetter::CreateStatic(&IsPropertyButtonEnabled, WeakPropertyEditor);
-		TAttribute<bool> IsEnabledAttribute = TAttribute<bool>::Create( EnabledDelegate );
+		TAttribute<bool>::FGetter IsPropertyButtonEnabledDelegate = TAttribute<bool>::FGetter::CreateStatic(&IsPropertyButtonEnabled, WeakPropertyEditor);
+		TAttribute<bool> IsEnabledAttribute = TAttribute<bool>::Create( IsPropertyButtonEnabledDelegate );
 
 		switch( ButtonType )
 		{
diff --git a/Engine/Source/Runtime/Engine/Private/Slate/SceneViewport.cpp b/Engine/Source/Runtime/Engine/Private/Slate/SceneViewport.cpp
--- a/Engine/Source/Runtime/Engine/Private/Slate/SceneViewport.cpp
+++ b/Engine/Source/Runtime/Engine/Private/Slate/SceneViewport.cpp
@@ -1437,7 +1437,6 @@ void FSceneViewport::InitDynamicRHI()
 
 		FWidgetPath WidgetPath;
 		TSharedPtr<SWindow> Window = FSlateApplication::Get().FindWidgetWindow(PinnedViewport.ToSharedRef(), WidgetPath);
-		TSharedPtr<FSlateRenderer> Renderer = FSlateApplication::Get().GetRenderer();
 		
 		WindowRenderTargetUpdate(Renderer.Get(), Window.Get());
 	}
