Fixing editor compilation issues (shadowing and etc).

diff --git a/Engine/Source/Developer/TreeMap/TreeMap.cpp b/Engine/Source/Developer/TreeMap/TreeMap.cpp
--- a/Engine/Source/Developer/TreeMap/TreeMap.cpp
+++ b/Engine/Source/Developer/TreeMap/TreeMap.cpp
@@ -279,7 +279,7 @@ void FTreeMap::MakeStandardTreeNode( const FTreeMapOptions& Options, FTreeMap::E
 }
 
 
-void FTreeMap::MakeSquarifiedTreeNode( const FTreeMapOptions& Options, const FTreeMapNodeRef& Node )
+void FTreeMap::MakeSquarifiedTreeNode( const FTreeMapOptions& Options, const FTreeMapNodeRef& InNode )
 {
 	// NOTE: This algorithm is explained in the paper titled, "Squarified Treemaps", by Mark Bruls, Kees Huizing, and Jarke J.van Wijk
 
@@ -417,8 +417,8 @@ void FTreeMap::MakeSquarifiedTreeNode( const FTreeMapOptions& Options, const FTr
 	};
 
 	// Squarify it!
-	auto ChildrenCopy = Node->Children;
-	FTreeMapRect SubRect = Node->Rect;
+	auto ChildrenCopy = InNode->Children;
+	FTreeMapRect SubRect = InNode->Rect;
 	do
 	{
 		const auto SubRectShortestSide = FMath::Min( SubRect.Size.X, SubRect.Size.Y );
@@ -557,10 +557,10 @@ FTreeMapNodeDataPtr ITreeMap::ParseOPMLToTreeMapData( const FString& OPMLFilePat
 			const auto& RootName = XmlRoot->GetTag();
 			if( RootName.Equals( TEXT( "opml" ), ESearchCase::IgnoreCase ) )
 			{
-				for( const auto& XmlNode : XmlRoot->GetChildrenNodes() )
+				for( const auto& OuterXmlNode : XmlRoot->GetChildrenNodes() )
 				{
-					const auto& NodeName = XmlNode->GetTag();
-					if( NodeName.Equals( TEXT( "body" ), ESearchCase::IgnoreCase ) )
+					const auto& OuterNodeName = OuterXmlNode->GetTag();
+					if( OuterNodeName.Equals( TEXT( "body" ), ESearchCase::IgnoreCase ) )
 					{
 						struct Local
 						{
@@ -635,7 +635,7 @@ FTreeMapNodeDataPtr ITreeMap::ParseOPMLToTreeMapData( const FString& OPMLFilePat
 						RootNodeData = MakeShareable( new FTreeMapNodeData() );
 						RootNodeData->Parent = NULL;
 						RootNodeData->Name = RootNodeName;
-						Local::RecursivelyCreateNodes( RootNodeData.ToSharedRef(), *XmlNode );
+						Local::RecursivelyCreateNodes( RootNodeData.ToSharedRef(), *OuterXmlNode );
 					}
 					else
 					{
diff --git a/Engine/Source/Editor/SizeMap/SSizeMap.cpp b/Engine/Source/Editor/SizeMap/SSizeMap.cpp
--- a/Engine/Source/Editor/SizeMap/SSizeMap.cpp
+++ b/Engine/Source/Editor/SizeMap/SSizeMap.cpp
@@ -127,7 +127,7 @@ namespace SizeMapInternals
 		{
 			// We ended up with bytes, so show a decimal number
 			BestSizeString = FString::Printf( TEXT( "%s %s" ), 
-				*FText::AsNumber( SizeInBytes ).ToString(), 
+				*FText::AsNumber( static_cast<uint64>(SizeInBytes) ).ToString(), 
 				*LOCTEXT( "Bytes", "bytes" ).ToString() );
 		}
 		else
@@ -160,7 +160,7 @@ namespace SizeMapInternals
 }
 
 
-void SSizeMap::GatherDependenciesRecursively( FAssetRegistryModule& AssetRegistryModule, TSharedPtr<FAssetThumbnailPool>& AssetThumbnailPool, TMap<FName, TSharedPtr<FTreeMapNodeData>>& VisitedAssetPackageNames, const TArray<FName>& AssetPackageNames, const TSharedPtr<FTreeMapNodeData>& Node, TSharedPtr<FTreeMapNodeData>& SharedRootNode, int32& NumAssetsWhichFailedToLoad )
+void SSizeMap::GatherDependenciesRecursively( FAssetRegistryModule& AssetRegistryModule, TSharedPtr<FAssetThumbnailPool>& InAssetThumbnailPool, TMap<FName, TSharedPtr<FTreeMapNodeData>>& VisitedAssetPackageNames, const TArray<FName>& AssetPackageNames, const TSharedPtr<FTreeMapNodeData>& Node, TSharedPtr<FTreeMapNodeData>& SharedRootNode, int32& NumAssetsWhichFailedToLoad )
 {
 	for( const FName AssetPackageName : AssetPackageNames )
 	{
@@ -288,7 +288,7 @@ void SSizeMap::GatherDependenciesRecursively( FAssetRegistryModule& AssetRegistr
 
 								
 						// Now visit all of the assets that we are referencing
-						GatherDependenciesRecursively( AssetRegistryModule, AssetThumbnailPool, VisitedAssetPackageNames, ReferencedAssetPackageNames, ChildTreeMapNode, SharedRootNode, NumAssetsWhichFailedToLoad );
+						GatherDependenciesRecursively( AssetRegistryModule, InAssetThumbnailPool, VisitedAssetPackageNames, ReferencedAssetPackageNames, ChildTreeMapNode, SharedRootNode, NumAssetsWhichFailedToLoad );
 
 					}
 					else
diff --git a/Engine/Source/Runtime/Linux/ALAudio/Private/ALAudioBuffer.cpp b/Engine/Source/Runtime/Linux/ALAudio/Private/ALAudioBuffer.cpp
--- a/Engine/Source/Runtime/Linux/ALAudio/Private/ALAudioBuffer.cpp
+++ b/Engine/Source/Runtime/Linux/ALAudio/Private/ALAudioBuffer.cpp
@@ -125,7 +125,6 @@ FALSoundBuffer* FALSoundBuffer::CreateNativeBuffer( FALAudioDevice* AudioDevice,
 		return nullptr;
 	}
 
-	FWaveModInfo WaveInfo;
 	Wave->InitAudioResource(AudioDevice->GetRuntimeFormat(Wave));
 	FALSoundBuffer* Buffer = nullptr;
 
