Editor - Fixing clipping artifacts on the Size Map tree widget.

diff --git a/Engine/Source/Developer/TreeMap/Public/STreeMap.h b/Engine/Source/Developer/TreeMap/Public/STreeMap.h
--- a/Engine/Source/Developer/TreeMap/Public/STreeMap.h
+++ b/Engine/Source/Developer/TreeMap/Public/STreeMap.h
@@ -46,7 +46,9 @@ public:
 		, _NestedContainerOuterPadding( 0.0f )
 		, _ContainerInnerPadding( 4.0f )
 		, _ChildContainerTextPadding( 2.0f )
-	{}
+	{
+		_Clipping = EWidgetClipping::ClipToBounds;
+	}
 
 		/** Sets whether the user can edit the tree map interactively by dragging nodes around and typing new node labels */
 		SLATE_ATTRIBUTE( bool, AllowEditing )
diff --git a/Engine/Source/Developer/TreeMap/STreeMap.cpp b/Engine/Source/Developer/TreeMap/STreeMap.cpp
--- a/Engine/Source/Developer/TreeMap/STreeMap.cpp
+++ b/Engine/Source/Developer/TreeMap/STreeMap.cpp
@@ -260,6 +260,9 @@ int32 STreeMap::OnPaint( const FPaintArgs& Args, const FGeometry& AllottedGeomet
 					const FSlateRect VisualClippingRect = TransformRect( AllottedGeometry.GetAccumulatedLayoutTransform(), FSlateRect( VisualPosition, VisualPosition + BlendedVisual.Size ) );
 					const auto VisualPaintGeometry = AllottedGeometry.ToPaintGeometry( VisualPosition, BlendedVisual.Size );
 					auto DrawColor = InWidgetStyle.GetColorAndOpacityTint() * ThisNodeBackground->TintColor.GetColor( InWidgetStyle ) * BlendedVisual.Color;
+
+					OutDrawElements.PushClip(FSlateClippingZone(VisualClippingRect));
+					
 					FSlateDrawElement::MakeBox(
 						OutDrawElements,
 						LayerId,
@@ -289,6 +292,8 @@ int32 STreeMap::OnPaint( const FPaintArgs& Args, const FGeometry& AllottedGeomet
 
 						const FSlateRect BackgroundClippingRect = VisualClippingRect.InsetBy( FMargin( 1 ) );
 
+						OutDrawElements.PushClip(FSlateClippingZone(BackgroundClippingRect));
+
 						// Draw the background brush
 						FSlateDrawElement::MakeBox(
 							OutDrawElements,
@@ -297,6 +302,8 @@ int32 STreeMap::OnPaint( const FPaintArgs& Args, const FGeometry& AllottedGeomet
 							BlendedVisual.NodeData->BackgroundBrush,
 							DrawEffects,
 							DrawColor );
+
+						OutDrawElements.PopClip();
 					}
 
 					const bool bIsHighlightPulseNode = HighlightPulseNode.IsValid() && HighlightPulseNode.Pin().Get() == BlendedVisual.NodeData;
@@ -307,6 +314,7 @@ int32 STreeMap::OnPaint( const FPaintArgs& Args, const FGeometry& AllottedGeomet
 						if( HighlightPulseAnimProgress >= 0.0f && HighlightPulseAnimProgress <= 1.0f )
 						{
 							DrawColor = FLinearColor( 1.0f, 1.0f, 1.0f, FMath::MakePulsatingValue( HighlightPulseAnimProgress, 6.0f, 0.5f ) );
+
 							FSlateDrawElement::MakeBox(
 								OutDrawElements,
 								LayerId,
@@ -317,6 +325,7 @@ int32 STreeMap::OnPaint( const FPaintArgs& Args, const FGeometry& AllottedGeomet
 						}
 					}
 
+					OutDrawElements.PopClip();
 				}
 			}
 		}
@@ -375,6 +384,8 @@ int32 STreeMap::OnPaint( const FPaintArgs& Args, const FGeometry& AllottedGeomet
 							TextClippingRect = TextClippingRect.InsetBy( FMargin( ChildContainerTextPadding, 0, ChildContainerTextPadding, 0 ) );
 							if( TextClippingRect.IsValid() )
 							{
+								OutDrawElements.PushClip(FSlateClippingZone(TextClippingRect));
+
 								// Name (first line)
 								float NameTextHeight = 0.0f;
 								if( !BlendedVisual.NodeData->Name.IsEmpty() )
@@ -439,6 +450,8 @@ int32 STreeMap::OnPaint( const FPaintArgs& Args, const FGeometry& AllottedGeomet
 										DrawEffects,
 										InWidgetStyle.GetColorAndOpacityTint() * VisualTextColor );
 								}
+
+								OutDrawElements.PopClip();
 							}
 						}
 					}
