Fixed DPI issues with SizeMap

diff --git a/Engine/Source/Developer/TreeMap/STreeMap.cpp b/Engine/Source/Developer/TreeMap/STreeMap.cpp
--- a/Engine/Source/Developer/TreeMap/STreeMap.cpp
+++ b/Engine/Source/Developer/TreeMap/STreeMap.cpp
@@ -379,11 +379,12 @@ int32 STreeMap::OnPaint( const FPaintArgs& Args, const FGeometry& AllottedGeomet
 						const FVector2D ScreenSpaceVisualSize( AllottedGeometry.Scale * BlendedVisual.Size );
 						if( ScreenSpaceVisualSize.X > 20 )
 						{
-							const auto VisualPaintGeometry = AllottedGeometry.ToPaintGeometry( VisualPosition, BlendedVisual.Size );
+							FGeometry VisualGeometry = AllottedGeometry.MakeChild(VisualPosition, BlendedVisual.Size);
+							const FPaintGeometry VisualPaintGeometry = VisualGeometry.ToPaintGeometry();
 
+							
 							// Clip the text to the visual's rectangle, with some extra inner padding to avoid overlapping the visual's border
-							auto TextClippingRect = FSlateRect( VisualPaintGeometry.DrawPosition, VisualPaintGeometry.DrawPosition + VisualPaintGeometry.GetLocalSize() );
-							TextClippingRect = TextClippingRect.IntersectionWith( MyCullingRect );
+							const FSlateRect TextClippingRect = VisualGeometry.GetLayoutBoundingRect();
 							TextClippingRect = TextClippingRect.InsetBy( FMargin( ChildContainerTextPadding, 0, ChildContainerTextPadding, 0 ) );
 							if( TextClippingRect.IsValid() )
 							{
