Fixed DPI issues with SizeMap (no compile errors this time)

diff --git a/Engine/Source/Developer/TreeMap/STreeMap.cpp b/Engine/Source/Developer/TreeMap/STreeMap.cpp
--- a/Engine/Source/Developer/TreeMap/STreeMap.cpp
+++ b/Engine/Source/Developer/TreeMap/STreeMap.cpp
@@ -384,7 +384,7 @@ int32 STreeMap::OnPaint( const FPaintArgs& Args, const FGeometry& AllottedGeomet
 
 							
 							// Clip the text to the visual's rectangle, with some extra inner padding to avoid overlapping the visual's border
-							const FSlateRect TextClippingRect = VisualGeometry.GetLayoutBoundingRect();
+							FSlateRect TextClippingRect = VisualGeometry.GetLayoutBoundingRect();
 							TextClippingRect = TextClippingRect.InsetBy( FMargin( ChildContainerTextPadding, 0, ChildContainerTextPadding, 0 ) );
 							if( TextClippingRect.IsValid() )
 							{
